Create the Subobject Data interface module. This provides an API that can be used to manipulate Subobjects such as adding/removing components, reparenting, and changing the scene heirarchy.

The SubobjectDataSubsystem is how users can manipulate a given object within Blueprints or Python scripting. Given a Uobject instance in a level or a UBlueprint asset, you can use GatherSubobjectData to get get an array of handles that you can use to manipulate that subobject.

This is what the new SubobjectEditor (Previously the SCS editor) will be using instead of having all its logic within slate code.

diff --git a/Engine/Source/Editor/SubobjectDataInterface/Private/BaseSubobjectDataFactory.cpp b/Engine/Source/Editor/SubobjectDataInterface/Private/BaseSubobjectDataFactory.cpp
new file mode 100644
--- /dev/null
+++ b/Engine/Source/Editor/SubobjectDataInterface/Private/BaseSubobjectDataFactory.cpp
@@ -0,0 +1,16 @@
+// Copyright Epic Games, Inc. All Rights Reserved.
+
+#include "BaseSubobjectDataFactory.h"
+#include "SubobjectData.h"
+
+TSharedPtr<FSubobjectData> FBaseSubobjectDataFactory::CreateSubobjectData(const FCreateSubobjectParams& Params)
+{
+	return TSharedPtr<FSubobjectData>(new FSubobjectData(Params.Context, Params.ParentHandle));
+}
+
+bool FBaseSubobjectDataFactory::ShouldCreateSubobjectData(const FCreateSubobjectParams& Params) const
+{
+	// By default, this will just return false to ensure that if anyone
+	// adds a custom subobject data factory that they won't have to change this function
+	return false;
+}
\ No newline at end of file
diff --git a/Engine/Source/Editor/SubobjectDataInterface/Private/ChildActorSubobjectData.cpp b/Engine/Source/Editor/SubobjectDataInterface/Private/ChildActorSubobjectData.cpp
new file mode 100644
--- /dev/null
+++ b/Engine/Source/Editor/SubobjectDataInterface/Private/ChildActorSubobjectData.cpp
@@ -0,0 +1,49 @@
+// Copyright Epic Games, Inc. All Rights Reserved.
+
+#include "ChildActorSubobjectData.h"
+
+FChildActorSubobjectData::FChildActorSubobjectData(UObject* ContextObject, const FSubobjectDataHandle& ParentHandle, const bool InbIsInheritedSCS)
+    : FInheritedSubobjectData(ContextObject, ParentHandle, InbIsInheritedSCS)
+{
+}
+
+FText FChildActorSubobjectData::GetDisplayName() const
+{
+	if(const UChildActorComponent* CAC = GetChildActorComponent())
+	{
+		return CAC->GetClass()->GetDisplayNameText();	
+	}
+	return FInheritedSubobjectData::GetDisplayName();
+}
+
+FText FChildActorSubobjectData::GetActorDisplayText() const
+{
+	if (const AActor* ChildActor = GetObject<AActor>())
+	{
+		return ChildActor->GetClass()->GetDisplayNameText();
+	}
+
+	return FInheritedSubobjectData::GetActorDisplayText();
+}
+
+bool FChildActorSubobjectData::IsChildActor() const
+{
+	return true;
+}
+
+bool FChildActorSubobjectData::CanDelete() const
+{
+	return false;
+}
+
+bool FChildActorSubobjectData::CanReparent() const
+{
+	// Can't reparent nodes within a child actor subtree
+	return false;
+}
+
+bool FChildActorSubobjectData::CanDuplicate() const
+{
+	// Duplication of Child actor components is not allowed.
+	return false;
+}	
\ No newline at end of file
diff --git a/Engine/Source/Editor/SubobjectDataInterface/Private/ChildSubobjectDataFactory.cpp b/Engine/Source/Editor/SubobjectDataInterface/Private/ChildSubobjectDataFactory.cpp
new file mode 100644
--- /dev/null
+++ b/Engine/Source/Editor/SubobjectDataInterface/Private/ChildSubobjectDataFactory.cpp
@@ -0,0 +1,27 @@
+// Copyright Epic Games, Inc. All Rights Reserved.
+
+#include "ChildSubobjectDataFactory.h"
+#include "ChildActorSubobjectData.h"
+#include "Components/ChildActorComponent.h"
+#include "SubobjectData.h"
+#include "Engine/SCS_Node.h"		// #TODO_BH  We need to remove this when the actual subobject refactor happens
+
+TSharedPtr<FSubobjectData> FChildSubobjectDataFactory::CreateSubobjectData(const FCreateSubobjectParams& Params)
+{
+	ensure(Params.Context);
+	return TSharedPtr<FChildActorSubobjectData>(new FChildActorSubobjectData(Params.Context, Params.ParentHandle, Params.bIsInheritedSCS));
+}
+
+bool FChildSubobjectDataFactory::ShouldCreateSubobjectData(const FCreateSubobjectParams& Params) const
+{
+	// Check for a BP added CAC
+	if(USCS_Node* SCS = Cast<USCS_Node>(Params.Context))
+	{
+		if(SCS->ComponentTemplate->IsA<UChildActorComponent>())
+		{
+			return true;
+		}
+	}
+	
+	return Params.Context && Params.Context->IsA<UChildActorComponent>();
+}
\ No newline at end of file
diff --git a/Engine/Source/Editor/SubobjectDataInterface/Private/InheritedSubobjectData.cpp b/Engine/Source/Editor/SubobjectDataInterface/Private/InheritedSubobjectData.cpp
new file mode 100644
--- /dev/null
+++ b/Engine/Source/Editor/SubobjectDataInterface/Private/InheritedSubobjectData.cpp
@@ -0,0 +1,48 @@
+// Copyright Epic Games, Inc. All Rights Reserved.
+
+#include "InheritedSubobjectData.h"
+#include "Kismet2/ComponentEditorUtils.h"
+
+FInheritedSubobjectData::FInheritedSubobjectData(UObject* ContextObject, const FSubobjectDataHandle& ParentHandle, const bool InbIsInheritedSCS)
+    : FSubobjectData(ContextObject, ParentHandle)
+	, bIsInheritedSCS(InbIsInheritedSCS)
+{
+}
+
+bool FInheritedSubobjectData::IsNativeComponent() const
+{
+	if (const UActorComponent* Template = GetComponentTemplate())
+    {
+    	return Template->CreationMethod == EComponentCreationMethod::Native && GetSCSNode() == nullptr;
+    }
+	
+	return false;
+}
+
+bool FInheritedSubobjectData::CanEdit() const
+{
+	if(IsComponent())
+	{
+		if(IsInstancedInheritedComponent())
+		{
+			const UActorComponent* Template = GetComponentTemplate();
+			return (Template ? Template->IsEditableWhenInherited() : false);
+		}
+		else if (!IsNativeComponent())
+		{
+			USCS_Node* SCS_Node = GetSCSNode();
+			return (SCS_Node != nullptr);
+		}
+		else if (const UActorComponent* ComponentTemplate = GetComponentTemplate())
+		{
+			return FComponentEditorUtils::GetPropertyForEditableNativeComponent(ComponentTemplate) != nullptr;
+		}
+	}
+	
+	return FSubobjectData::CanEdit();
+}
+
+bool FInheritedSubobjectData::IsInheritedSCSNode() const
+{
+	return bIsInheritedSCS;
+}
diff --git a/Engine/Source/Editor/SubobjectDataInterface/Private/InheritedSubobjectDataFactory.cpp b/Engine/Source/Editor/SubobjectDataInterface/Private/InheritedSubobjectDataFactory.cpp
new file mode 100644
--- /dev/null
+++ b/Engine/Source/Editor/SubobjectDataInterface/Private/InheritedSubobjectDataFactory.cpp
@@ -0,0 +1,23 @@
+// Copyright Epic Games, Inc. All Rights Reserved.
+
+#include "InheritedSubobjectDataFactory.h"
+#include "InheritedSubobjectData.h"
+
+TSharedPtr<FSubobjectData> FInheritedSubobjectDataFactory::CreateSubobjectData(const FCreateSubobjectParams& Params)
+{
+	return TSharedPtr<FInheritedSubobjectData>(new FInheritedSubobjectData(Params.Context, Params.ParentHandle, Params.bIsInheritedSCS));
+}
+
+bool FInheritedSubobjectDataFactory::ShouldCreateSubobjectData(const FCreateSubobjectParams& Params) const
+{
+	if(UActorComponent* Component = Cast<UActorComponent>(Params.Context))
+	{
+		// Create an inherited subobject data
+		if(Params.bIsInheritedSCS || Component->CreationMethod == EComponentCreationMethod::Native)
+		{
+			return true;
+		}
+	}
+
+	return false;
+}
\ No newline at end of file
diff --git a/Engine/Source/Editor/SubobjectDataInterface/Private/SubobjectData.cpp b/Engine/Source/Editor/SubobjectDataInterface/Private/SubobjectData.cpp
new file mode 100644
--- /dev/null
+++ b/Engine/Source/Editor/SubobjectDataInterface/Private/SubobjectData.cpp
@@ -0,0 +1,1040 @@
+// Copyright Epic Games, Inc. All Rights Reserved.
+
+#include "SubobjectData.h"
+#include "SubobjectDataSubsystem.h"		// For iterating child if we need to
+#include "Engine/Blueprint.h"			// Casting to UBlueprint
+#include "Components/ChildActorComponent.h"
+#include "GameFramework/Actor.h"
+#include "Kismet2/ComponentEditorUtils.h"
+#include "Kismet2/BlueprintEditorUtils.h"
+#include "ComponentAssetBroker.h"			// FComponentAssetBrokerage
+#include "Engine/World.h"					// Finding the instance of a component on an actor
+#include "Engine/InheritableComponentHandler.h"
+#include "Engine/SCS_Node.h"		// #TODO_BH  We need to remove this when the actual subobject refactor happens
+
+#define LOCTEXT_NAMESPACE "SubobjectData"
+
+FSubobjectData::FSubobjectData()
+	: WeakObjectPtr(nullptr)
+	, ParentObjectHandle(FSubobjectDataHandle::InvalidHandle)
+	, SCSNodePtr(nullptr)
+{
+	// By Default the handle will be invalid, it will only
+	// be generated if we are given an object.
+}
+
+FSubobjectData::FSubobjectData(UObject* ContextObject, const FSubobjectDataHandle& InParentHandle)
+	: WeakObjectPtr(ContextObject)
+	, ParentObjectHandle(InParentHandle)
+	, SCSNodePtr(nullptr)
+{
+	AttemptToSetSCSNode();
+}
+
+FSubobjectData::~FSubobjectData()
+{
+	ChildrenHandles.Empty();
+	WeakObjectPtr = nullptr;
+}
+
+bool FSubobjectData::CanEdit() const
+{
+	// Actors are editable by default
+	if(const AActor* ActorContext = GetObject<AActor>())
+	{
+		return true;
+	}
+	return false;
+}
+
+bool FSubobjectData::CanDelete() const
+{
+	// Components can be deleted if they are not inherited or the scene root
+	if (const UActorComponent* ComponentTemplate = GetComponentTemplate())
+	{
+		return !IsInheritedComponent() && !IsDefaultSceneRoot() && !IsChildActor();
+	}
+	
+	// Otherwise, it can't be deleted
+	return false;
+}
+
+bool FSubobjectData::CanDuplicate() const
+{
+	return CanCopy();
+}
+
+bool FSubobjectData::CanCopy() const
+{
+	return FComponentEditorUtils::CanCopyComponent(GetComponentTemplate());
+}
+
+bool FSubobjectData::CanReparent() const
+{	
+	if(IsComponent())
+	{
+		if(GetSCSNode() != nullptr)
+		{
+			return !IsInstancedInheritedComponent();
+		}
+		
+		return !IsInheritedComponent() && !IsDefaultSceneRoot() && IsSceneComponent();
+	}
+	
+	return false;
+}
+
+bool FSubobjectData::CanRename() const
+{
+	if(!IsValid())
+	{
+		return false;
+	}
+	
+	// You can rename instance actors in the level editor, but you cannot rename
+	// the default actors in blueprints. 
+	if(IsActor())
+	{
+		return IsInstancedActor() && !IsChildActor();
+	}
+
+	// You can rename within the BP editor, but not if if this subobject
+	// is on an instance in the level
+	if(GetSCSNode() != nullptr)
+	{
+		return !IsInstancedInheritedComponent();
+	}
+	
+	return !IsInheritedComponent() && !IsDefaultSceneRoot() && !IsChildActor();
+}
+
+const UObject* FSubobjectData::GetObjectForBlueprint(UBlueprint* Blueprint) const
+{
+	const bool bCanEdit = CanEdit();
+	// We have to deal with the ICH (Inherited Component Handler) for components 
+	if(IsComponent() && bCanEdit && !IsNativeComponent() && IsInheritedSCSNode())
+	{
+		UActorComponent* OverriddenComponent = nullptr;
+		FComponentKey Key(GetSCSNode());
+		
+		const bool bBlueprintCanOverrideComponentFromKey = Key.IsValid()
+		    && Blueprint
+		    && Blueprint->ParentClass
+		    && Blueprint->ParentClass->IsChildOf(Key.GetComponentOwner());
+		
+		if (bBlueprintCanOverrideComponentFromKey)
+		{
+			const bool bCreateIfNecessary = true;
+			UInheritableComponentHandler* InheritableComponentHandler = Blueprint->GetInheritableComponentHandler(bCreateIfNecessary);
+			if (InheritableComponentHandler)
+			{
+				OverriddenComponent = InheritableComponentHandler->GetOverridenComponentTemplate(Key);
+				if (!OverriddenComponent && bCreateIfNecessary)
+				{
+					OverriddenComponent = InheritableComponentHandler->CreateOverridenComponentTemplate(Key);
+				}
+			}
+		}
+		return OverriddenComponent;
+	}
+	
+	// If this not a component, then we can simply return the current object as long as it's editable
+	return bCanEdit ? WeakObjectPtr.Get() : nullptr;
+}
+
+bool FSubobjectData::IsInstancedComponent() const
+{
+	// Check the flags on the component (RF_ClassDefaultObject | RF_ArchetypeObject)
+	const UActorComponent* Component = GetComponentTemplate();
+	return Component && !Component->IsTemplate();
+}
+
+UActorComponent* FSubobjectData::FindMutableComponentInstanceInActor(const AActor* InActor) const
+{
+	const USCS_Node* SCS_Node = GetSCSNode();
+	const UActorComponent* ComponentTemplate = GetComponentTemplate();
+
+	UActorComponent* ComponentInstance = nullptr;
+	if (InActor)
+	{
+		if (SCS_Node)
+		{
+			FName VariableName = SCS_Node->GetVariableName();
+			if (VariableName != NAME_None)
+			{
+				UWorld* World = InActor->GetWorld();
+				FObjectPropertyBase* Property = FindFProperty<FObjectPropertyBase>(InActor->GetClass(), VariableName);
+				if (Property != nullptr)
+				{
+					// Return the component instance that's stored in the property with the given variable name
+					ComponentInstance = Cast<UActorComponent>(Property->GetObjectPropertyValue_InContainer(InActor));
+				}
+				else if (World != nullptr && World->WorldType == EWorldType::EditorPreview)
+				{
+					// If this is the preview actor, return the cached component instance that's being used for the preview actor prior to recompiling the Blueprint
+					ComponentInstance = SCS_Node->EditorComponentInstance.Get();
+				}
+			}
+		}
+		else if (ComponentTemplate)
+		{
+			TInlineComponentArray<UActorComponent*> Components;
+			InActor->GetComponents(Components);
+			ComponentInstance = FComponentEditorUtils::FindMatchingComponent(ComponentTemplate, Components); 
+		}
+	}
+
+	return ComponentInstance;
+}
+
+UBlueprint* FSubobjectData::GetBlueprint() const
+{
+	// If this object is a BP, we can just return that
+	if (UBlueprint* BP = Cast<UBlueprint>(WeakObjectPtr.Get()))
+	{
+		return BP;
+	}
+	// If it is an actor, then we can get the BP from the UClass
+	else if (const AActor* DefaultActor = GetObject<AActor>())
+	{
+		return UBlueprint::GetBlueprintFromClass(DefaultActor->GetClass());
+	}
+	// For components, we need to get the blueprint from their owning actor or the SCS
+	else if(IsComponent())
+	{
+		if (const USCS_Node* SCS_Node = GetSCSNode())
+		{
+			if (const USimpleConstructionScript* SCS = SCS_Node->GetSCS())
+			{
+				return SCS->GetBlueprint();
+			}
+		}
+		else if(const UActorComponent* ActorComponent = GetComponentTemplate())
+		{
+			if (const AActor* Actor = ActorComponent->GetOwner())
+			{
+				return UBlueprint::GetBlueprintFromClass(Actor->GetClass());
+			}
+		}
+	}
+
+	return nullptr;
+}
+
+FString FSubobjectData::GetDisplayString(bool bShowNativeComponentNames /* = true */) const
+{
+	FName VariableName = GetVariableName();
+
+	const UActorComponent* ComponentTemplate = GetComponentTemplate();
+	
+	UBlueprint* Blueprint = GetBlueprint();
+	UClass* VariableOwner = (Blueprint != nullptr) ? Blueprint->SkeletonGeneratedClass : nullptr;
+	FProperty* VariableProperty = FindFProperty<FProperty>(VariableOwner, VariableName);
+
+	bool const bHasValidVarName = (VariableName != NAME_None);
+	bool const bIsArrayVariable = bHasValidVarName && (VariableOwner != nullptr) &&
+		VariableProperty && VariableProperty->IsA<FArrayProperty>();
+
+	// Only display SCS node variable names in the tree if they have not been autogenerated
+	if (ComponentTemplate && bHasValidVarName && !bIsArrayVariable)
+	{
+		const bool bIsNative = IsNativeComponent();
+		const bool bIsInherited = IsInheritedComponent();
+		const bool bIsInstance = IsInstancedComponent();
+		const bool bIsBlueprintInstanceInherited = bIsInherited && bIsInstance;
+		
+		// Inherited/Native components will have "Name (Inherited)" as their display
+		if ((bIsNative || bIsInherited) && bShowNativeComponentNames)
+		{
+			FStringFormatNamedArguments Args;
+			Args.Add(TEXT("VarName"), VariableProperty && VariableProperty->IsNative() ? VariableProperty->GetDisplayNameText().ToString() : VariableName.ToString());
+			FString CompName = TEXT(" (") + ComponentTemplate->GetName() + TEXT(")");
+			Args.Add(TEXT("CompName"), bIsNative || IsInstancedInheritedComponent() ? CompName : TEXT(""));
+			return FString::Format(TEXT("{VarName}{CompName}"), Args);
+		}
+		else
+		{
+			return VariableName.ToString();
+		}
+	}
+	else if (ComponentTemplate != nullptr)
+	{
+		return ComponentTemplate->GetFName().ToString();
+	}
+	else if (const AActor* DefaultActor = GetObject<AActor>())
+	{
+		FString Name;
+		if (Blueprint != nullptr && !IsInstancedActor())
+		{
+			Blueprint->GetName(Name);
+		}
+		else
+		{
+			Name = DefaultActor->GetActorLabel();
+		}
+		
+		FStringFormatNamedArguments Args;
+		Args.Add(TEXT("ActorName"), Name);
+		return FString::Format(TEXT("{ActorName}"), Args);
+	}
+	// If the context is a simple UObject, then we can get it's name
+	else if (const UObject* Context = GetObject())
+	{
+		FString Name;
+		Context->GetName(Name);
+
+		return Name;
+	}
+	// Anything else will be unknown!
+	else
+	{
+		FString UnnamedString = LOCTEXT("UnnamedToolTip", "Unnamed").ToString();
+		FString NativeString = IsNativeComponent() ? LOCTEXT("NativeToolTip", "Native ").ToString() : TEXT("");
+
+		if (ComponentTemplate != nullptr)
+		{
+			return FString::Printf(TEXT("[%s %s%s]"), *UnnamedString, *NativeString, *ComponentTemplate->GetClass()->GetName());
+		}
+		else
+		{
+			return FString::Printf(TEXT("[%s %s]"), *UnnamedString, *NativeString);
+		}
+	}
+}
+
+FText FSubobjectData::GetDisplayNameContextModifiers(bool bShowNativeComponentNames) const
+{
+	if(IsActor())
+	{
+		if(IsChildActor())
+		{
+			return LOCTEXT("ActorContext_ChildActor", "(Child Actor)");
+		}
+		else
+		{
+			if (const AActor* DefaultActor = GetObject<AActor>())
+			{
+				if (const UBlueprint* Blueprint = UBlueprint::GetBlueprintFromClass(DefaultActor->GetClass()))
+				{
+					return LOCTEXT("ActorContext_self", "(Self)");
+				}
+				else
+				{
+					return LOCTEXT("ActorContext_Instance", "(Instance)");
+				}
+			}
+		}
+	}
+	else if(const UActorComponent* Template = GetComponentTemplate())
+	{
+		FName VariableName = GetVariableName();
+		UBlueprint* Blueprint = GetBlueprint();
+		UClass* VariableOwner = (Blueprint != nullptr) ? Blueprint->SkeletonGeneratedClass : nullptr;
+		FProperty* VariableProperty = FindFProperty<FProperty>(VariableOwner, VariableName);
+
+		bool const bHasValidVarName = (VariableName != NAME_None);
+		bool const bIsArrayVariable = bHasValidVarName && (VariableOwner != nullptr) && VariableProperty && VariableProperty->IsA<FArrayProperty>();
+		const bool bIsBlueprintInstanceInherited = GetSCSNode() != nullptr && IsInstancedInheritedComponent();
+
+		if(bIsBlueprintInstanceInherited)
+		{
+			FStringFormatNamedArguments Args;
+            Args.Add(TEXT("InheritedText"), TEXT("(Inherited)"));
+            return FText::FromString(FString::Format(TEXT("{InheritedText}"), Args));
+		}
+		
+		// Only display SCS node variable names in the tree if they have not been autogenerated
+		if (bHasValidVarName && !bIsArrayVariable)
+		{
+			const bool bIsNative = IsNativeComponent();
+			const bool bIsInherited = IsInheritedComponent();
+			const bool bIsInstance = IsInstancedComponent();
+
+			if ((bIsNative || bIsInherited) && bShowNativeComponentNames)
+			{		
+				FStringFormatNamedArguments Args;
+				Args.Add(TEXT("InheritedText"), bIsInherited || bIsBlueprintInstanceInherited ? TEXT("(Inherited)") : TEXT(""));
+				return FText::FromString(FString::Format(TEXT("{InheritedText}"), Args));
+			}
+		}
+	}
+
+	return FText::GetEmpty();
+}
+
+FText FSubobjectData::GetDisplayName() const
+{
+	if (const AActor* DefaultActor = GetObject<AActor>())
+	{
+		FString Name;
+		UBlueprint* Blueprint = UBlueprint::GetBlueprintFromClass(DefaultActor->GetClass());
+		if (Blueprint != nullptr && !IsInstancedActor())
+		{
+			Blueprint->GetName(Name);
+		}
+		else
+		{
+			Name = DefaultActor->GetActorLabel();
+		}
+		return FText::FromString(Name);
+	}
+	else if (const UObject* Context = GetObject())
+	{
+		FString Name;
+		Context->GetName(Name);
+
+		return FText::FromString(Name);
+	}
+
+	return LOCTEXT("GetDisplayNameNotOverridden", "Unknown Subobject");
+}
+
+FName FSubobjectData::GetVariableName() const
+{
+	FName VariableName = NAME_None;
+
+	const USCS_Node* SCS_Node = GetSCSNode();
+	const UActorComponent* ComponentTemplate = GetComponentTemplate();
+
+	if (SCS_Node != nullptr)
+	{
+		// Use the same variable name as is obtained by the compiler
+		VariableName = SCS_Node->GetVariableName();
+	}
+	else if (ComponentTemplate != nullptr)
+	{
+		// Try to find the component anchor variable name (first looks for an exact match then scans for any matching variable that points to the archetype in the CDO)
+		VariableName = FComponentEditorUtils::FindVariableNameGivenComponentInstance(ComponentTemplate);
+	}
+
+	return VariableName;
+}
+
+FText FSubobjectData::GetSocketName() const
+{
+	if (USCS_Node* SCSNode = GetSCSNode())
+	{
+		return FText::FromName(SCSNode->AttachToName);
+	}
+	return FText::GetEmpty();
+}
+
+FName FSubobjectData::GetSocketFName() const
+{
+	if (USCS_Node* SCSNode = GetSCSNode())
+	{
+		return SCSNode->AttachToName;
+	}
+	return NAME_None;
+}
+
+bool FSubobjectData::HasValidSocket() const
+{
+	return GetSCSNode() != nullptr;
+}
+
+void FSubobjectData::SetSocketName(FName InNewName)
+{
+	if(USCS_Node* SCS = GetSCSNode())
+	{
+		SCS->AttachToName = InNewName;
+	}
+}
+
+void FSubobjectData::SetupAttachment(FName SocketName, const FSubobjectDataHandle& AttachParentHandle)
+{
+	USceneComponent* SceneComponentTemplate = Cast<USceneComponent>(GetMutableComponentTemplate());
+
+	FSubobjectData* AttachParentData = AttachParentHandle.GetData();
+	USceneComponent* AttachParent = AttachParentData ?
+		Cast<USceneComponent>(AttachParentData->GetMutableComponentTemplate()) :
+		SceneComponentTemplate->GetAttachParent();
+	
+	SceneComponentTemplate->SetupAttachment(AttachParent, NAME_None);
+	if(USCS_Node* SCS_Node = GetSCSNode())
+	{
+		SCS_Node->Modify();
+		SCS_Node->AttachToName = NAME_None;
+	}
+}
+
+FSubobjectDataHandle FSubobjectData::GetRootSubobject() const
+{
+	FSubobjectDataHandle Current = ParentObjectHandle;
+	while(Current.IsValid() && Current.GetSharedDataPtr()->HasParent())
+	{
+		Current = Current.GetSharedDataPtr()->GetParentHandle();
+	}
+
+	return Current;
+}
+
+bool FSubobjectData::HasChild(const FSubobjectDataHandle& InChildHandle) const
+{
+	for(const FSubobjectDataHandle& MyChildHandle : ChildrenHandles)
+	{
+		if(MyChildHandle == InChildHandle)
+		{
+			return true;
+		}
+	}
+	return false;
+}
+
+FSubobjectDataHandle FSubobjectData::FindChild(const FSubobjectDataHandle& InChildHandle) const
+{
+	for(const FSubobjectDataHandle& MyChildHandle : ChildrenHandles)
+	{
+		if(MyChildHandle == InChildHandle)
+		{
+			return MyChildHandle;
+		}		
+	}
+
+	return FSubobjectDataHandle::InvalidHandle;
+}
+
+FSubobjectDataHandle FSubobjectData::FindChildByObject(UObject* ContextObject) const
+{
+	if (!ContextObject)
+	{
+		return FSubobjectDataHandle::InvalidHandle;
+	}
+
+	if (USubobjectDataSubsystem* System = USubobjectDataSubsystem::Get())
+	{
+		for (const FSubobjectDataHandle& CurrentChild : ChildrenHandles)
+		{
+			if (FSubobjectData* ChildData = CurrentChild.GetData())
+			{
+				if (ChildData->GetObject() == ContextObject)
+				{
+					return CurrentChild;
+				}
+			}
+		}
+	}
+
+	return FSubobjectDataHandle::InvalidHandle;
+}
+
+bool FSubobjectData::AddChildHandleOnly(const FSubobjectDataHandle& InHandle)
+{
+	// If we already have this child, then don't both with adding it
+	if (HasChild(InHandle) || !InHandle.IsValid() || InHandle == Handle)
+	{
+		return false;
+	}
+
+	ChildrenHandles.Add(InHandle);
+	if(FSubobjectData* NewChildData = InHandle.GetData())
+	{
+		NewChildData->SetParentHandle(Handle);		
+	}
+
+	return true;
+}
+
+bool FSubobjectData::RemoveChildHandleOnly(const FSubobjectDataHandle& InHandle)
+{
+	if (HasChild(InHandle))
+	{
+		ChildrenHandles.Remove(InHandle);
+		return true;
+	}
+	return false;
+}
+
+FText FSubobjectData::GetAssetName() const
+{
+	UActorComponent* Template = GetMutableComponentTemplate();
+
+	FText AssetName = LOCTEXT("None", "None");
+	if (Template)
+	{
+		UObject* Asset = FComponentAssetBrokerage::GetAssetFromComponent(Template);
+		if (Asset != nullptr)
+		{
+			AssetName = FText::FromString(Asset->GetName());
+		}
+	}
+
+	return AssetName;
+}
+
+FText FSubobjectData::GetAssetPath() const
+{
+	FText AssetName = LOCTEXT("None", "None");
+
+	UActorComponent* Template = GetMutableComponentTemplate();
+
+	if (Template)
+	{
+		UObject* Asset = FComponentAssetBrokerage::GetAssetFromComponent(Template);
+		if (Asset != nullptr)
+		{
+			AssetName = FText::FromString(Asset->GetPathName());
+		}
+	}
+
+	return AssetName;
+}
+
+bool FSubobjectData::IsAssetVisible() const
+{
+	UActorComponent* Template = GetMutableComponentTemplate();
+
+	if (Template && FComponentAssetBrokerage::SupportsAssets(Template))
+	{
+		return true;
+	}
+
+	return false;
+}
+
+FText FSubobjectData::GetMobilityToolTipText() const
+{
+	FText MobilityToolTip;
+    
+    if (const USceneComponent* SceneComponentTemplate = Cast<USceneComponent>(GetComponentTemplate()))
+    {
+    	if (SceneComponentTemplate->Mobility == EComponentMobility::Movable)
+    	{
+    		MobilityToolTip = LOCTEXT("MovableMobilityTooltip", "Movable");
+    	}
+    	else if (SceneComponentTemplate->Mobility == EComponentMobility::Stationary)
+    	{
+    		MobilityToolTip = LOCTEXT("StationaryMobilityTooltip", "Stationary");
+    	}
+    	else if (SceneComponentTemplate->Mobility == EComponentMobility::Static)
+    	{
+    		MobilityToolTip = LOCTEXT("StaticMobilityTooltip", "Static");
+    	}
+    	else
+    	{
+    		// make sure we're the mobility type we're expecting (we've handled Movable & Stationary)
+    		ensureMsgf(false, TEXT("Unhandled mobility type [%d], is this a new type that we don't handle here?"), SceneComponentTemplate->Mobility.GetValue());
+    		MobilityToolTip = LOCTEXT("UnknownMobilityTooltip", "Component with unknown mobility");
+    	}
+    }
+    else
+    {
+    	MobilityToolTip = LOCTEXT("NoMobilityTooltip", "Non-scene component");
+    }
+
+	return MobilityToolTip;
+}
+
+FText FSubobjectData::GetComponentEditorOnlyTooltipText() const
+{
+	FText ComponentType = LOCTEXT("ComponentEditorOnlyFalse", "False");
+	
+	if (IsComponent())
+	{
+		if (const UActorComponent* Template = GetComponentTemplate())
+		{
+			UBlueprint* Blueprint = GetBlueprint();
+			FObjectProperty* Prop = Blueprint ? FindFProperty<FObjectProperty>(Blueprint->SkeletonGeneratedClass, GetVariableName()) : nullptr;
+
+			if(Template->bIsEditorOnly || (Prop && Prop->HasAnyPropertyFlags(CPF_EditorOnly)))
+			{
+				ComponentType = LOCTEXT("ComponentEditorOnlyTrue", "True");
+			}
+		}
+	}
+	
+	return ComponentType;
+}
+
+FText FSubobjectData::GetIntroducedInToolTipText() const
+{
+	FText IntroducedInTooltip = LOCTEXT("IntroducedInThisBPTooltip", "this class");
+
+	if (IsInheritedComponent())
+	{
+		if (const UActorComponent* ComponentTemplate = GetComponentTemplate())
+		{
+			UClass* BestClass = nullptr;
+			AActor* OwningActor = ComponentTemplate->GetOwner();
+
+			if (IsNativeComponent() && (OwningActor != nullptr))
+			{
+				for (UClass* TestClass = OwningActor->GetClass(); TestClass != AActor::StaticClass(); TestClass = TestClass->GetSuperClass())
+				{
+					if (FindComponentInstanceInActor(Cast<AActor>(TestClass->GetDefaultObject())))
+					{
+						BestClass = TestClass;
+					}
+					else
+					{
+						break;
+					}
+				}
+			}
+			else if (!IsNativeComponent())
+			{
+				USCS_Node* SCSNode = GetSCSNode();
+
+				if ((SCSNode == nullptr) && (OwningActor != nullptr))
+				{
+					SCSNode = FindSCSNodeForInstance(ComponentTemplate, OwningActor->GetClass());
+				}
+
+				if (SCSNode != nullptr)
+				{
+					if (UBlueprint* OwningBP = SCSNode->GetSCS()->GetBlueprint())
+					{
+						BestClass = OwningBP->GeneratedClass;
+					}
+				}
+				else if (OwningActor != nullptr)
+				{
+					if (UBlueprint* OwningBP = UBlueprint::GetBlueprintFromClass(OwningActor->GetClass()))
+					{
+						BestClass = OwningBP->GeneratedClass;
+					}
+				}
+			}
+
+			if (BestClass == nullptr)
+			{
+				if (ComponentTemplate->IsCreatedByConstructionScript()) 
+				{
+					IntroducedInTooltip = LOCTEXT("IntroducedInUnknownError", "Unknown Blueprint Class (via an Add Component call)");
+				} 
+				else 
+				{
+					IntroducedInTooltip = LOCTEXT("IntroducedInNativeError", "Unknown native source (via C++ code)");
+				}
+			}
+			else if (IsInstancedComponent() && ComponentTemplate->CreationMethod == EComponentCreationMethod::Native && !ComponentTemplate->HasAnyFlags(RF_DefaultSubObject))
+			{
+				IntroducedInTooltip = FText::Format(LOCTEXT("IntroducedInCPPErrorFmt", "{0} (via C++ code)"), FBlueprintEditorUtils::GetFriendlyClassDisplayName(BestClass));
+			}
+			else if (IsInstancedComponent() && ComponentTemplate->CreationMethod == EComponentCreationMethod::UserConstructionScript)
+			{
+				IntroducedInTooltip = FText::Format(LOCTEXT("IntroducedInUCSErrorFmt", "{0} (via an Add Component call)"), FBlueprintEditorUtils::GetFriendlyClassDisplayName(BestClass));
+			}
+			else
+			{
+				IntroducedInTooltip = FBlueprintEditorUtils::GetFriendlyClassDisplayName(BestClass);
+			}
+		}
+		else
+		{
+			IntroducedInTooltip = LOCTEXT("IntroducedInNoTemplateError", "[no component template found]");
+		}
+	}
+	else if (IsInstancedComponent())
+	{
+		IntroducedInTooltip = LOCTEXT("IntroducedInThisActorInstanceTooltip", "this actor instance");
+	}
+
+	return IntroducedInTooltip;
+}
+
+FText FSubobjectData::GetActorDisplayText() const
+{
+	if (const AActor* DefaultActor = GetObject<AActor>())
+	{
+		FString Name;
+		UBlueprint* Blueprint = UBlueprint::GetBlueprintFromClass(DefaultActor->GetClass());
+		if (Blueprint != nullptr && !IsInstancedActor())
+		{
+			Blueprint->GetName(Name);
+		}
+		else
+		{
+			Name = DefaultActor->GetActorLabel();
+		}
+		return FText::FromString(Name);
+	}
+	
+	return FText::GetEmpty();
+}
+
+FLinearColor FSubobjectData::GetColorTintForIcon() const
+{
+	// A blue-ish tint
+	static const FLinearColor InheritedBlueprintComponentColor(0.08f, 0.35f, 0.6f);
+	
+	static const FLinearColor InstancedInheritedBlueprintComponentColor(0.08f, 0.35f, 0.6f);
+	// A green-ish tint
+	static const FLinearColor InheritedNativeComponentColor(0.7f, 0.9f, 0.7f);
+	
+	static const FLinearColor IntroducedHereColor(FLinearColor::White);
+
+	if (IsInheritedComponent())
+	{
+		// Native C++ components will be tinted green
+		if (IsNativeComponent())
+		{
+			return InheritedNativeComponentColor;
+		}
+		else if (IsInstancedComponent())
+		{
+			return InstancedInheritedBlueprintComponentColor;
+		}
+		else
+		{
+			return InheritedBlueprintComponentColor;
+		}
+	}
+	// If we have an SCS but are not Inherited, then this is just a regular blueprint and we should be blue
+	else if(GetSCSNode() != nullptr)
+	{
+		// If it's inherited BP color then it should be blue (i.e. this is a BP component that came from a BP generated class)
+		if(IsInheritedSCSNode() || IsInstancedComponent())
+		{
+			return InheritedBlueprintComponentColor;			
+		}
+		// Otherwise this is just a regular SCS node inside of the BP editor
+		else
+		{
+			return IntroducedHereColor;
+		}
+	}
+
+	// By default, this node will appear white to represent being introduced here
+	return IntroducedHereColor;
+}
+
+bool FSubobjectData::IsInstancedActor() const
+{
+	if (const AActor* Actor = GetObject<AActor>())
+	{
+		return !Actor->IsTemplate();
+	}
+
+	return false;
+}
+
+bool FSubobjectData::IsNativeComponent() const
+{
+	if (const UActorComponent* Template = GetComponentTemplate())
+	{
+		return Template->CreationMethod == EComponentCreationMethod::Native && GetSCSNode() == nullptr;
+	}
+	
+	return false;
+}
+
+bool FSubobjectData::IsInheritedComponent() const
+{
+	// If this component is not created via an instance, then it's an inherited component
+
+	// This covers a component that is added via blueprints
+	if(GetSCSNode() != nullptr)
+	{
+		return IsInheritedSCSNode();
+	}
+	else if (const UActorComponent* ComponentTemplate = GetComponentTemplate())
+	{
+		return IsNativeComponent() || ComponentTemplate->CreationMethod != EComponentCreationMethod::Instance;
+	}
+
+	return false;
+}
+
+bool FSubobjectData::IsSceneComponent() const
+{
+	return Cast<USceneComponent>(GetComponentTemplate()) != nullptr;
+}
+
+bool FSubobjectData::IsRootComponent() const
+{
+	const UActorComponent* ComponentTemplate = GetComponentTemplate();
+	
+	if(!ComponentTemplate)
+	{
+		return false;
+	}
+	
+	const AActor* CDO = ComponentTemplate ? ComponentTemplate->GetOwner() : nullptr;
+	
+	if(CDO && (IsInstancedComponent() || IsInheritedComponent()))
+	{
+		return CDO->GetRootComponent() == ComponentTemplate;
+	}
+	
+	bool bIsRoot = true;
+	if(USCS_Node* SCS_Node = GetSCSNode())
+	{
+		if(USimpleConstructionScript* SCS = SCS_Node->GetSCS())
+		{
+			// Evaluate to TRUE if we have an SCS node reference, it is contained in the SCS root set and does not have an external parent
+			bIsRoot = SCS->GetRootNodes().Contains(SCS_Node) && SCS_Node->ParentComponentOrVariableName == NAME_None;
+		}
+	}
+	else if(ComponentTemplate && CDO)
+	{
+		// Evaluate to TRUE if we have a valid component reference that matches the native root component
+		bIsRoot = (ComponentTemplate == CDO->GetRootComponent());
+	}
+	
+	return bIsRoot;
+}
+
+bool FSubobjectData::IsDefaultSceneRoot() const
+{
+	// If this is a scene component and is instanced, then it will have a specific name
+	const USceneComponent* SceneComponent = Cast<USceneComponent>(GetComponentTemplate());
+	if (SceneComponent && !SceneComponent->IsTemplate())
+	{
+		return SceneComponent->GetFName() == USceneComponent::GetDefaultSceneRootVariableName();
+	}
+	// If this isn't a scene component, then we can check an SCS node for a flag. 
+	else if (const USCS_Node* SCS_Node = GetSCSNode())
+	{
+		USimpleConstructionScript* SCS = SCS_Node->GetSCS();
+		if (SCS != nullptr)
+		{
+			return SCS_Node == SCS->GetDefaultSceneRootNode();
+		}
+	}
+
+	// Nothing else can be a DefaultSceneRoot
+	return false;
+}
+
+bool FSubobjectData::IsComponent() const
+{
+	// Check if we are pointing to a component
+	return GetComponentTemplate() != nullptr;
+}
+
+bool FSubobjectData::IsChildActor() const
+{
+	return false;
+}
+
+bool FSubobjectData::IsRootActor() const
+{
+	// This is the root actor if it points to an AActor and has no parent
+	if(const AActor* Actor = GetObject<AActor>())
+	{
+		return !ParentObjectHandle.IsValid();
+	}
+	return false;
+}
+
+bool FSubobjectData::IsActor() const
+{
+	return GetObject<AActor>() != nullptr;
+}
+
+bool FSubobjectData::IsInstancedInheritedComponent() const
+{
+	if(!IsComponent())
+	{
+		return false;
+	}
+	FSubobjectDataHandle CurrentHandle = ParentObjectHandle;
+	USubobjectDataSubsystem* System = USubobjectDataSubsystem::Get();
+	FSubobjectData* CurrentData = CurrentHandle.GetData();
+			
+	while(CurrentHandle.IsValid() && CurrentData && !CurrentData->IsActor())
+	{
+		CurrentHandle = CurrentData->GetParentHandle();
+		CurrentData = CurrentHandle.GetData(); 
+	}
+
+	return CurrentData && CurrentData->IsInstancedActor();
+}
+
+bool FSubobjectData::IsAttachedTo(const FSubobjectDataHandle& InHandle) const
+{
+	if (USubobjectDataSubsystem* System = USubobjectDataSubsystem::Get())
+	{
+		FSubobjectDataHandle TestParentHandle = ParentObjectHandle;
+
+		while (TestParentHandle.IsValid())
+		{
+			if (TestParentHandle == InHandle)
+			{
+				return true;
+			}
+			const FSubobjectData* TestParentData = TestParentHandle.GetData();
+			TestParentHandle = TestParentData ? TestParentData->GetParentHandle() : FSubobjectDataHandle::InvalidHandle;
+		}
+	}
+
+	return false;
+}
+
+AActor* FSubobjectData::GetMutableActorContext()
+{
+	if(AActor* Actor = GetMutableObject<AActor>())
+	{
+		return Actor;
+	}
+	else if(UActorComponent* Component = GetMutableComponentTemplate())
+	{
+		return Component->GetOwner();
+	}
+	else if(UBlueprint* BP = GetBlueprint())
+	{
+		return BP->GeneratedClass ? BP->GeneratedClass->GetDefaultObject<AActor>() : nullptr;
+	}
+	return nullptr;
+}
+
+USCS_Node* FSubobjectData::FindSCSNodeForInstance(const UActorComponent* InstanceComponent, UClass* ClassToSearch)
+{
+	if ((ClassToSearch != nullptr) && InstanceComponent->IsCreatedByConstructionScript())
+	{
+		for (UClass* TestClass = ClassToSearch; TestClass->ClassGeneratedBy != nullptr; TestClass = TestClass->GetSuperClass())
+		{
+			if (UBlueprint* TestBP = Cast<UBlueprint>(TestClass->ClassGeneratedBy))
+			{
+				if (TestBP->SimpleConstructionScript != nullptr)
+				{
+					if (USCS_Node* Result = TestBP->SimpleConstructionScript->FindSCSNode(InstanceComponent->GetFName()))
+					{
+						return Result;
+					}
+				}
+			}
+		}
+	}
+
+	return nullptr;
+}
+
+bool FSubobjectData::AttemptToSetSCSNode()
+{
+	if (USCS_Node* PossibleSCS = Cast<USCS_Node>(WeakObjectPtr.Get()))
+	{
+		WeakObjectPtr = PossibleSCS->ComponentTemplate;
+		SCSNodePtr = PossibleSCS;
+		return true;
+	}
+	// If this is an instanced component, then we can find it's SCS node
+	else if (IsInstancedComponent())
+	{
+		const UActorComponent* Template = GetComponentTemplate();
+		if (Template->GetOwner())
+		{
+			SCSNodePtr = FindSCSNodeForInstance(Template, Template->GetOwner()->GetClass());
+			if (SCSNodePtr.IsValid())
+			{
+				WeakObjectPtr = SCSNodePtr->ComponentTemplate;
+				return true;
+			}
+		}
+	}
+	return false;
+}
+
+USCS_Node* FSubobjectData::GetSCSNode(bool bEvenIfPendingKill) const
+{
+	// @todo Deprecate everything related to SCS nodes that could possibly be public facing.
+	return SCSNodePtr.IsValid() ? SCSNodePtr.Get() : Cast<USCS_Node>(WeakObjectPtr.Get(bEvenIfPendingKill));
+}
+
+bool FSubobjectData::IsInheritedSCSNode() const
+{
+	return false;
+}
+
+#undef LOCTEXT_NAMESPACE
diff --git a/Engine/Source/Editor/SubobjectDataInterface/Private/SubobjectDataBlueprintFunctionLibrary.cpp b/Engine/Source/Editor/SubobjectDataInterface/Private/SubobjectDataBlueprintFunctionLibrary.cpp
new file mode 100644
--- /dev/null
+++ b/Engine/Source/Editor/SubobjectDataInterface/Private/SubobjectDataBlueprintFunctionLibrary.cpp
@@ -0,0 +1,12 @@
+// Copyright Epic Games, Inc. All Rights Reserved.
+
+#include "SubobjectDataBlueprintFunctionLibrary.h"
+
+void USubobjectDataBlueprintFunctionLibrary::GetData(const FSubobjectDataHandle& DataHandle, FSubobjectData& OutData)
+{
+	TSharedPtr<FSubobjectData> DataPtr = DataHandle.GetSharedDataPtr();
+	if(DataPtr.IsValid())
+	{
+		OutData = *DataPtr.Get();
+	}
+}
diff --git a/Engine/Source/Editor/SubobjectDataInterface/Private/SubobjectDataHandle.cpp b/Engine/Source/Editor/SubobjectDataInterface/Private/SubobjectDataHandle.cpp
new file mode 100644
--- /dev/null
+++ b/Engine/Source/Editor/SubobjectDataInterface/Private/SubobjectDataHandle.cpp
@@ -0,0 +1,5 @@
+// Copyright Epic Games, Inc. All Rights Reserved.
+
+#include "SubobjectDataHandle.h"
+
+const FSubobjectDataHandle FSubobjectDataHandle::InvalidHandle = FSubobjectDataHandle();
\ No newline at end of file
diff --git a/Engine/Source/Editor/SubobjectDataInterface/Private/SubobjectDataInterfaceModule.cpp b/Engine/Source/Editor/SubobjectDataInterface/Private/SubobjectDataInterfaceModule.cpp
new file mode 100644
--- /dev/null
+++ b/Engine/Source/Editor/SubobjectDataInterface/Private/SubobjectDataInterfaceModule.cpp
@@ -0,0 +1,20 @@
+// Copyright Epic Games, Inc. All Rights Reserved.
+
+#include "SubobjectDataInterfaceModule.h"
+
+#define LOCTEXT_NAMESPACE "FSubobjectDataInterfaceModule"
+
+void FSubobjectDataInterfaceModule::StartupModule()
+{
+	// This code will execute after your module is loaded into memory; the exact timing is specified in the .uplugin file per-module
+}
+
+void FSubobjectDataInterfaceModule::ShutdownModule()
+{
+	// This function may be called during shutdown to clean up your module.  For modules that support dynamic reloading,
+	// we call this function before unloading the module.	
+}
+
+#undef LOCTEXT_NAMESPACE
+
+IMPLEMENT_MODULE(FSubobjectDataInterfaceModule, SubobjectDataInterface)
\ No newline at end of file
diff --git a/Engine/Source/Editor/SubobjectDataInterface/Private/SubobjectDataSubsystem.cpp b/Engine/Source/Editor/SubobjectDataInterface/Private/SubobjectDataSubsystem.cpp
new file mode 100644
--- /dev/null
+++ b/Engine/Source/Editor/SubobjectDataInterface/Private/SubobjectDataSubsystem.cpp
@@ -0,0 +1,2084 @@
+// Copyright Epic Games, Inc. All Rights Reserved.
+
+#include "SubobjectDataSubsystem.h"
+#include "UnrealEngine.h"				// GEngine for subsystems
+#include "Editor/EditorEngine.h"		// FActorLabelUtilities
+#include "ComponentAssetBroker.h"		// FComponentAssetBrokerage
+#include "ChildActorSubobjectData.h"
+#include "InheritedSubobjectData.h"
+#include "BaseSubobjectDataFactory.h"
+#include "ChildSubobjectDataFactory.h"
+#include "InheritedSubobjectDataFactory.h"
+
+#include "BlueprintEditorSettings.h"	// bHideConstructionScriptComponentsInDetailsView
+#include "Kismet2/ComponentEditorUtils.h"
+#include "Kismet2/KismetEditorUtilities.h"	// IsClassAllowed
+#include "Kismet2/Kismet2NameValidators.h"
+#include "Kismet2/BlueprintEditorUtils.h"
+#include "ClassViewerFilter.h"
+#include "AssetRegistryModule.h"
+#include "ScopedTransaction.h"
+
+#include "GameProjectGenerationModule.h"	// Adding new component classes
+#include "GameProjectUtils.h"
+#include "SourceCodeNavigation.h"
+#include "AddToProjectConfig.h"
+
+#include "Serialization/ObjectWriter.h"
+#include "Serialization/ObjectReader.h"
+
+#include "Kismet2/CompilerResultsLog.h"		// Adding compiler errors to nodes that have their variables deleted
+#include "K2Node_ComponentBoundEvent.h"
+
+#include "Engine/SCS_Node.h"		// #TODO_BH  We need to remove this when the actual subobject refactor happens
+
+#define LOCTEXT_NAMESPACE "SubobjectDataInterface"
+
+DEFINE_LOG_CATEGORY_STATIC(LogSubobjectSubsystem, Log, All);
+
+//////////////////////////////////////////////
+// USubobjectDataSubsystem
+
+void USubobjectDataSubsystem::Initialize(FSubsystemCollectionBase& Collection)
+{
+	FactoryManager = new FSubobjectFactoryManager();
+	check(FactoryManager);
+
+	FactoryManager->RegisterFactory(MakeShared<FBaseSubobjectDataFactory>());
+	FactoryManager->RegisterFactory(MakeShared<FChildSubobjectDataFactory>());
+	FactoryManager->RegisterFactory(MakeShared<FInheritedSubobjectDataFactory>());
+}
+
+void USubobjectDataSubsystem::Deinitialize()
+{
+	if(ensure(FactoryManager))
+	{
+		delete FactoryManager;
+		FactoryManager = nullptr;
+	}
+}
+
+USubobjectDataSubsystem* USubobjectDataSubsystem::Get()
+{
+	return GEngine->GetEngineSubsystem<USubobjectDataSubsystem>();
+}
+
+void USubobjectDataSubsystem::K2_GatherSubobjectDataForBlueprint(UBlueprint* Context, TArray<FSubobjectDataHandle>& OutArray)
+{
+	// Return the current CDO that was last generated for the class
+	if (Context != nullptr && Context->GeneratedClass != nullptr)
+	{
+		GatherSubobjectData(Context->GeneratedClass->GetDefaultObject(), OutArray);
+	}
+}
+
+void USubobjectDataSubsystem::K2_GatherSubobjectDataForInstance(AActor* Context, TArray<FSubobjectDataHandle>& OutArray)
+{
+	return GatherSubobjectData(Context, OutArray);
+}
+
+void USubobjectDataSubsystem::GatherSubobjectData(UObject* Context, TArray<FSubobjectDataHandle>& OutArray)
+{
+	if (!Context)
+	{
+		UE_LOG(LogSubobjectSubsystem, Warning, TEXT("Could not gather subobject data, there was no context given!"));
+		return;
+	}
+
+	const AActor* ActorContext = Cast<AActor>(Context);
+	if(!ActorContext)
+	{
+		UE_LOG(LogSubobjectSubsystem, Warning, TEXT("Could not gather subobject data, the given context was not an actor!"));
+		return;
+	}
+	
+	OutArray.Reset();
+
+	FSubobjectDataHandle RootActorHandle = CreateSubobjectData(Context);
+	OutArray.Add(RootActorHandle);
+	const FSubobjectData* RootActorDataPtr = RootActorHandle.GetData();
+	check(RootActorDataPtr);
+
+	// Is the root an actor? 
+	const bool bIsInstanced = RootActorDataPtr->IsInstancedActor();
+
+	// If the actor is not instanced, then we are dealing with a BP class
+	if (!bIsInstanced)
+	{
+		// get all the components
+		TArray<UActorComponent*> Components;
+		ActorContext->GetComponents(Components);
+
+		USceneComponent* RootComponent = ActorContext->GetRootComponent();
+		FSubobjectDataHandle RootComponentHandle;
+		if (RootComponent != nullptr)
+		{
+			Components.Remove(RootComponent);
+			RootComponentHandle = FactoryCreateSubobjectDataWithParent(RootComponent, RootActorHandle);
+			OutArray.Add(RootComponentHandle);
+		}
+
+		// The components array will be populated with any natively added components
+		// from the constructor/ObjectInitalizer
+		// These components will all be attached to the root component if it exists
+		for (UActorComponent* Component : Components)
+		{
+			FSubobjectDataHandle NewComponentHandle = FactoryCreateSubobjectDataWithParent(
+					/* Context = */ Component,
+					/* ParentHandle = */ RootComponentHandle.IsValid() ? RootComponentHandle : RootActorHandle
+			);
+			ensureMsgf(NewComponentHandle.IsValid(), TEXT("Gathering of native components failed!"));
+			OutArray.Add(NewComponentHandle);
+		}
+
+		// If it's a Blueprint-generated class, also get the inheritance stack
+		TArray<UBlueprint*> ParentBPStack;
+		UBlueprint::GetBlueprintHierarchyFromClass(ActorContext->GetClass(), ParentBPStack);
+
+		// Add the full SCS tree node hierarchy (including SCS nodes inherited from parent blueprints)
+		for (int32 StackIndex = ParentBPStack.Num() - 1; StackIndex >= 0; --StackIndex)
+		{
+			if (ParentBPStack[StackIndex]->SimpleConstructionScript != nullptr)
+			{
+				const TArray<USCS_Node*>& SCS_RootNodes = ParentBPStack[StackIndex]->SimpleConstructionScript->GetRootNodes();
+				for (int32 NodeIndex = 0; NodeIndex < SCS_RootNodes.Num(); ++NodeIndex)
+				{
+					USCS_Node* SCS_Node = SCS_RootNodes[NodeIndex];
+					check(SCS_Node != nullptr);
+
+					// Create a new subobject whose parent 
+					FSubobjectDataHandle NewHandle;
+					FSubobjectData* NewData = nullptr;
+					// If this SCS node has a parent component, then add it to that
+					if (SCS_Node->ParentComponentOrVariableName != NAME_None)
+					{
+						USceneComponent* ParentComponent = SCS_Node->GetParentComponentTemplate(ParentBPStack[0]);
+						if (ParentComponent != nullptr)
+						{
+							// The parent component will already be in the out array if it is valid, so search for it
+							FSubobjectDataHandle ParentHandle;
+							for(const FSubobjectDataHandle& CurHandle : OutArray)
+							{
+								if(TSharedPtr<FSubobjectData> DataPtr = CurHandle.GetSharedDataPtr())
+								{
+									if (DataPtr->GetComponentTemplate() == ParentComponent)
+									{
+										ParentHandle = CurHandle;
+										break;
+									}
+								}
+							}
+							
+							if (ensure(ParentHandle.IsValid()))
+							{
+								NewHandle = FactoryCreateInheritedBpSubobject(SCS_Node, RootComponentHandle.IsValid() ? RootComponentHandle : RootActorHandle, /* bIsInherited = */ StackIndex > 0, OutArray);
+							}
+						}
+					}
+					// Otherwise add a node parents to the root actor
+					else
+					{
+						NewHandle = FactoryCreateInheritedBpSubobject(SCS_Node, RootComponentHandle.IsValid() ? RootComponentHandle : RootActorHandle, /* bIsInherited = */ StackIndex > 0, OutArray);
+					}
+					NewData = NewHandle.GetData();
+					
+					// Only necessary to do the following for inherited nodes (StackIndex > 0).
+					if (NewData && StackIndex > 0)
+					{
+						// This call creates ICH override templates for the current Blueprint. Without this, the parent node
+						// search above can fail when attempting to match an inherited node in the tree via component template.
+						NewData->GetObjectForBlueprint(ParentBPStack[0]);
+						for (FSubobjectDataHandle ChildHandle : NewData->GetChildrenHandles())
+						{
+							FSubobjectData* ChildData = ChildHandle.GetData();
+							if (ensure(ChildData != nullptr))
+							{
+								ChildData->GetObjectForBlueprint(ParentBPStack[0]);
+							}
+						}
+					}
+				}
+			}
+		}
+	}
+	// Otherwise, this is an actor instance in a level
+	else
+	{
+		// Get the full set of instanced components
+		TSet<UActorComponent*> ComponentsToAdd(ActorContext->GetComponents());
+
+		const bool bHideConstructionScriptComponentsInDetailsView = GetDefault<UBlueprintEditorSettings>()->bHideConstructionScriptComponentsInDetailsView;
+		auto ShouldAddInstancedActorComponent = [bHideConstructionScriptComponentsInDetailsView](UActorComponent* ActorComp, USceneComponent* ParentSceneComp)
+		{
+			// Exclude nested DSOs attached to BP-constructed instances, which are not mutable.
+			return (ActorComp != nullptr
+				&& (!ActorComp->IsVisualizationComponent())
+				&& (ActorComp->CreationMethod != EComponentCreationMethod::UserConstructionScript || !bHideConstructionScriptComponentsInDetailsView)
+				&& (ParentSceneComp == nullptr || !ParentSceneComp->IsCreatedByConstructionScript() || !ActorComp->HasAnyFlags(RF_DefaultSubObject)))
+				&& (ActorComp->CreationMethod != EComponentCreationMethod::Native || FComponentEditorUtils::GetPropertyForEditableNativeComponent(ActorComp));
+		};
+
+		// Filter the components by their visibility
+		for (TSet<UActorComponent*>::TIterator It(ComponentsToAdd.CreateIterator()); It; ++It)
+		{
+			UActorComponent* ActorComp = *It;
+			USceneComponent* SceneComp = Cast<USceneComponent>(ActorComp);
+			USceneComponent* ParentSceneComp = SceneComp != nullptr ? SceneComp->GetAttachParent() : nullptr;
+			if (!ShouldAddInstancedActorComponent(ActorComp, ParentSceneComp))
+			{
+				It.RemoveCurrent();
+			}
+		}
+
+		TFunction<void(USceneComponent*, FSubobjectDataHandle)> AddInstancedComponentsRecursive = [&](USceneComponent* Component, FSubobjectDataHandle ParentHandle)
+		{
+			if (Component != nullptr)
+			{
+				for (USceneComponent* ChildComponent : Component->GetAttachChildren())
+				{
+					if (ComponentsToAdd.Contains(ChildComponent) && ChildComponent->GetOwner() == Component->GetOwner())
+					{
+						ComponentsToAdd.Remove(ChildComponent);
+						FSubobjectDataHandle NewParentHandle = FactoryCreateSubobjectDataWithParent(ChildComponent, ParentHandle);
+						OutArray.Add(NewParentHandle);
+						
+						AddInstancedComponentsRecursive(ChildComponent, NewParentHandle);
+					}
+				}
+			}
+		};
+
+		USceneComponent* RootComponent = ActorContext->GetRootComponent();
+
+		// Add the root component first
+		if (RootComponent != nullptr)
+		{
+			// We want this to be first every time, so remove it from the set of components that will be added later
+			ComponentsToAdd.Remove(RootComponent);
+
+			// Add the root component first
+			FSubobjectDataHandle RootCompHandle = FactoryCreateSubobjectDataWithParent(RootComponent, RootActorHandle);
+			OutArray.Add(RootCompHandle);
+
+			// Recursively add
+			AddInstancedComponentsRecursive(RootComponent, RootCompHandle);
+		}
+
+		// Sort components by type (always put scene components first in the tree)
+		ComponentsToAdd.Sort([](const UActorComponent& A, const UActorComponent& /* B */)
+		{
+			return A.IsA<USceneComponent>();
+		});
+
+		// Now add any remaining instanced owned components not already added above. This will first add any
+		// unattached scene components followed by any instanced non-scene components owned by the Actor instance.
+		for (UActorComponent* ActorComp : ComponentsToAdd)
+		{
+			// Create new subobject data with the original data as their parent.
+			OutArray.Add(FactoryCreateSubobjectDataWithParent(ActorComp, RootActorHandle));
+		}
+	}
+}
+
+void USubobjectDataSubsystem::FindAllSubobjectData(FSubobjectData* InData, TSet<FSubobjectData*>& OutVisited) const
+{
+	if(!InData || OutVisited.Contains(InData))
+	{
+		return;
+	}
+	
+	OutVisited.Add(InData);
+		
+	for (FSubobjectDataHandle ChildHandle : InData->GetChildrenHandles())
+	{
+		FindAllSubobjectData(ChildHandle.GetData(), OutVisited);
+	}
+}
+
+bool USubobjectDataSubsystem::K2_FindSubobjectDataFromHandle(const FSubobjectDataHandle& Handle, FSubobjectData& OutData) const
+{
+	const FSubobjectData* Data = Handle.GetData();
+	if (Data)
+	{
+		OutData = *Data;
+		return true;
+	}
+	else
+	{
+		return false;
+	}
+}
+
+FSubobjectDataHandle USubobjectDataSubsystem::FindHandleForObject(const FSubobjectDataHandle& Context, const UObject* ObjectToFind, UBlueprint* BPContext /* = nullptr */) const
+{
+	if(Context.IsValid())
+	{
+		if(const UActorComponent* ComponentToFind = Cast<UActorComponent>(ObjectToFind))
+		{
+			// If the given component instance is not already an archetype object
+			if(BPContext && !ComponentToFind->IsTemplate())
+			{
+				// Get the component owner's class object
+				check(ComponentToFind->GetOwner() != nullptr);
+				UClass* OwnerClass = ComponentToFind->GetOwner()->GetClass();
+
+				// If the given component is one that's created during Blueprint construction
+				if (ComponentToFind->IsCreatedByConstructionScript())
+				{
+					// Check the entire Class hierarchy for the node
+					TArray<UBlueprintGeneratedClass*> ParentBPStack;
+					UBlueprint::GetBlueprintHierarchyFromClass(OwnerClass, ParentBPStack);
+
+					for(int32 StackIndex = ParentBPStack.Num() - 1; StackIndex >= 0; --StackIndex)
+					{
+						USimpleConstructionScript* ParentSCS = ParentBPStack[StackIndex] ? ParentBPStack[StackIndex]->SimpleConstructionScript.Get() : nullptr;
+						if (ParentSCS)
+						{
+							// Attempt to locate an SCS node with a variable name that matches the name of the given component
+							for (USCS_Node* SCS_Node : ParentSCS->GetAllNodes())
+							{
+								check(SCS_Node != nullptr);
+								if (SCS_Node->GetVariableName() == ComponentToFind->GetFName())
+								{
+									// We found a match; redirect to the component archetype instance that may be associated with a tree node
+									ObjectToFind = SCS_Node->ComponentTemplate;
+									break;
+								}
+							}
+						}
+					}
+				}
+				else
+				{
+					// Get the class default object
+					const AActor* CDO = Cast<AActor>(OwnerClass->GetDefaultObject());
+					if (CDO)
+					{
+						// Iterate over the Components array and attempt to find a component with a matching name
+						for (UActorComponent* ComponentTemplate : CDO->GetComponents())
+						{
+							if (ComponentTemplate && ComponentTemplate->GetFName() == ComponentToFind->GetFName())
+							{
+								// We found a match; redirect to the component archetype instance that may be associated with a tree node
+								ObjectToFind = ComponentTemplate;
+								break;
+							}
+						}
+					}
+				}
+			}
+		}
+
+		TSet<FSubobjectData*> OutData;
+		FindAllSubobjectData(Context.GetSharedDataPtr().Get(), OutData);
+
+		for(FSubobjectData* CurData : OutData)
+		{
+			if(CurData->GetObject() == ObjectToFind)
+			{
+				return CurData->GetHandle();
+			}
+		}
+	}
+	
+	return FSubobjectDataHandle::InvalidHandle;
+}
+
+class FComponentClassParentFilter : public IClassViewerFilter
+{
+public:
+	FComponentClassParentFilter(const TSubclassOf<UActorComponent>& InComponentClass) : ComponentClass(InComponentClass) {}
+
+	virtual bool IsClassAllowed(const FClassViewerInitializationOptions& InInitOptions, const UClass* InClass, TSharedRef< FClassViewerFilterFuncs > InFilterFuncs) override
+	{
+		return InClass->IsChildOf(ComponentClass);
+	}
+
+	virtual bool IsUnloadedClassAllowed(const FClassViewerInitializationOptions& InInitOptions, const TSharedRef< const IUnloadedBlueprintData > InUnloadedClassData, TSharedRef< FClassViewerFilterFuncs > InFilterFuncs) override
+	{
+		return InUnloadedClassData->IsChildOf(ComponentClass);
+	}
+
+	TSubclassOf<UActorComponent> ComponentClass;
+};
+
+typedef FComponentClassParentFilter FNativeComponentClassParentFilter;
+
+class FBlueprintComponentClassParentFilter : public FComponentClassParentFilter
+{
+public:
+	FBlueprintComponentClassParentFilter(const TSubclassOf<UActorComponent>& InComponentClass) : FComponentClassParentFilter(InComponentClass) {}
+
+	virtual bool IsClassAllowed(const FClassViewerInitializationOptions& InInitOptions, const UClass* InClass, TSharedRef< FClassViewerFilterFuncs > InFilterFuncs) override
+	{
+		return FComponentClassParentFilter::IsClassAllowed(InInitOptions, InClass, InFilterFuncs) && FKismetEditorUtilities::CanCreateBlueprintOfClass(InClass);
+	}
+};
+
+UClass* USubobjectDataSubsystem::CreateNewCPPComponent(TSubclassOf<UActorComponent> ComponentClass, const FString& NewClassPath, const FString& NewClassName)
+{
+	UClass* NewClass = nullptr;
+
+	if (ComponentClass && !NewClassName.IsEmpty() && !NewClassPath.IsEmpty())
+	{
+		FString HeaderFilePath;
+		FString CppFilePath;
+		FText FailReason;
+
+		TSharedPtr<FModuleContextInfo> SelectedModuleInfo = MakeShareable(new FModuleContextInfo());
+		FNewClassInfo NewClassInfo(ComponentClass);
+
+		// Attempt to add new source files to the project
+		const TSet<FString>& DisallowedHeaderNames = FSourceCodeNavigation::GetSourceFileDatabase().GetDisallowedHeaderNames();
+		const GameProjectUtils::EAddCodeToProjectResult AddCodeResult = 
+			GameProjectUtils::AddCodeToProject(
+			NewClassName, 
+			NewClassPath, 
+			*SelectedModuleInfo, 
+			NewClassInfo, 
+			DisallowedHeaderNames, 
+			/*out*/ HeaderFilePath,
+			/*out*/ CppFilePath,
+			/*out*/ FailReason
+		);
+
+		if (AddCodeResult == GameProjectUtils::EAddCodeToProjectResult::Succeeded)
+		{			
+			FString AddedClassName = FString::Printf(TEXT("/Script/%s.%s"), *SelectedModuleInfo->ModuleName, *NewClassName);
+			NewClass = LoadClass<UActorComponent>(nullptr, *AddedClassName, nullptr, LOAD_None, nullptr);
+		}
+		else
+		{
+			UE_LOG(LogSubobjectSubsystem, Error, TEXT("Failed to create a new CPP component: %s"), *FailReason.ToString())
+		}
+	}
+
+	return NewClass;
+}
+
+UClass* USubobjectDataSubsystem::CreateNewBPComponent(TSubclassOf<UActorComponent> ComponentClass, const FString& NewClassPath, const FString& NewClassName)
+{
+	UClass* NewClass = nullptr;
+	if (ComponentClass && !NewClassName.IsEmpty() && !NewClassPath.IsEmpty())
+	{
+		const FString PackagePath = NewClassPath / NewClassName;
+
+		if (UPackage* Package = CreatePackage(*PackagePath))
+		{
+			// Create and init a new Blueprint
+			UBlueprint* NewBP = FKismetEditorUtilities::CreateBlueprint(ComponentClass, Package, FName(*NewClassName), BPTYPE_Normal, UBlueprint::StaticClass(), UBlueprintGeneratedClass::StaticClass());
+			if (NewBP)
+			{
+				// Notify the asset registry
+				FAssetRegistryModule::AssetCreated(NewBP);
+
+				// Mark the package dirty...
+				Package->MarkPackageDirty();
+				NewClass = NewBP->GeneratedClass;
+			}
+		}
+	}
+	return NewClass;
+}
+
+static void SaveSCSCurrentState(USimpleConstructionScript* SCSObj)
+{
+	if (SCSObj)
+	{
+		SCSObj->SaveToTransactionBuffer();
+	}
+}
+
+static void ConformTransformRelativeToParent(USceneComponent* SceneComponentTemplate, const USceneComponent* ParentSceneComponent)
+{
+	// If we find a match, calculate its new position relative to the scene root component instance in its current scene
+	FTransform ComponentToWorld(SceneComponentTemplate->GetRelativeRotation(), SceneComponentTemplate->GetRelativeLocation(), SceneComponentTemplate->GetRelativeScale3D());
+	FTransform ParentToWorld = SceneComponentTemplate->GetAttachSocketName() != NAME_None ? ParentSceneComponent->GetSocketTransform(SceneComponentTemplate->GetAttachSocketName(), RTS_World) : ParentSceneComponent->GetComponentToWorld();
+	FTransform RelativeTM = ComponentToWorld.GetRelativeTransform(ParentToWorld);
+
+	// Store new relative location value (if not set to absolute)
+	if (!SceneComponentTemplate->IsUsingAbsoluteLocation())
+	{
+		SceneComponentTemplate->SetRelativeLocation_Direct(RelativeTM.GetTranslation());
+	}
+
+	// Store new relative rotation value (if not set to absolute)
+	if (!SceneComponentTemplate->IsUsingAbsoluteRotation())
+	{
+		SceneComponentTemplate->SetRelativeRotation_Direct(RelativeTM.Rotator());
+	}
+
+	// Store new relative scale value (if not set to absolute)
+	if (!SceneComponentTemplate->IsUsingAbsoluteScale())
+	{
+		SceneComponentTemplate->SetRelativeScale3D_Direct(RelativeTM.GetScale3D());
+	}
+}
+
+FSubobjectDataHandle USubobjectDataSubsystem::FindSceneRootForSubobject(const FSubobjectDataHandle& InHandle) const
+{
+	if(!InHandle.IsValid())
+	{
+		return FSubobjectDataHandle::InvalidHandle;
+	}
+	
+	FSubobjectDataHandle ActorHandle = InHandle;
+	FSubobjectData* ActorData = ActorHandle.GetData();
+	
+	// If the given handle is not an actor, then we can walk it's parent chain back up until we find it
+	while (ActorData && !ActorData->IsActor())
+	{
+		ActorHandle = ActorData->GetParentHandle();
+		ActorData = ActorHandle.GetData();
+	}
+
+	// If the given handle is an actor, then we have to use it's first scene component
+	if(ensure(ActorData && ActorData->IsActor()))
+	{
+		TArray<FSubobjectDataHandle> ChildHandles = ActorData->GetChildrenHandles();
+		for(const FSubobjectDataHandle& ChildHandle : ChildHandles)
+		{
+			FSubobjectData* ChildData = ChildHandle.GetData();
+			if(ChildData && ChildData->IsDefaultSceneRoot())
+			{
+				return ChildHandle;
+			}
+		}
+	}
+
+	// Last resort is to just return the last known actor handle
+	return ActorHandle;
+}
+
+FSubobjectDataHandle USubobjectDataSubsystem::FindParentForNewSubobject(const UObject* NewSubobject, const FSubobjectDataHandle& SelectedParent)
+{
+	FSubobjectDataHandle TargetParentHandle = SelectedParent;
+	check(TargetParentHandle.IsValid());
+	FSubobjectData* TargetParentData = TargetParentHandle.GetData();
+
+	// If the current selection belongs to a child actor template, move the target to its outer component node.
+	while(TargetParentHandle.IsValid() && TargetParentData->IsChildActor())
+	{
+		TargetParentHandle = TargetParentData->GetParentHandle();
+		TargetParentData = TargetParentHandle.GetData();
+		check(TargetParentData);
+	}
+	
+	if(const USceneComponent* NewSceneComponent = Cast<USceneComponent>(NewSubobject))
+	{
+		if(TargetParentData)
+		{
+			if (TargetParentData->IsActor())
+			{
+				AActor* TargetActor = TargetParentData->GetMutableObject<AActor>();
+				check(TargetActor);
+				USceneComponent* TargetRootComp = TargetActor ? TargetActor->GetDefaultAttachComponent() : nullptr;
+				// A component should attach to the target scene component 
+				FSubobjectDataHandle RootComponentHandle = FactoryCreateSubobjectDataWithParent(
+					TargetRootComp ? Cast<UObject>(TargetRootComp) : Cast<UObject>(TargetActor),
+					TargetParentData->GetHandle()
+				);
+				
+				if (RootComponentHandle.IsValid())
+				{
+					TargetParentHandle = RootComponentHandle;
+					const USceneComponent* CastTargetToSceneComponent = Cast<USceneComponent>(TargetParentData->GetComponentTemplate());
+					if (CastTargetToSceneComponent == nullptr || !NewSceneComponent->CanAttachAsChild(CastTargetToSceneComponent, NAME_None))
+					{
+						TargetParentHandle = FindSceneRootForSubobject(SelectedParent); // Default to SceneRoot
+					}
+				}
+			}
+			else if(TargetParentData->IsComponent())
+			{
+				const USceneComponent* CastTargetToSceneComponent = Cast<USceneComponent>(TargetParentData->GetComponentTemplate());
+				if (CastTargetToSceneComponent == nullptr || !NewSceneComponent->CanAttachAsChild(CastTargetToSceneComponent, NAME_None))
+				{
+					TargetParentHandle = FindSceneRootForSubobject(SelectedParent); // Default to SceneRoot
+				}
+			}
+		}
+		else
+		{
+			TargetParentHandle = FindSceneRootForSubobject(SelectedParent);
+		}
+	}
+	else
+	{
+		if (TargetParentData->IsValid())
+		{
+			TargetParentHandle = FindSceneRootForSubobject(SelectedParent);
+		}
+		else
+		{
+			TargetParentHandle = SelectedParent;
+		}
+
+		check(TargetParentHandle.IsValid() && TargetParentData->IsActor());
+	}
+	
+	return TargetParentHandle;
+}
+
+FSubobjectDataHandle USubobjectDataSubsystem::AddNewSubobject(const FAddNewSubobjectParams& Params, FText& FailReason)
+{
+	FSubobjectDataHandle NewDataHandle = FSubobjectDataHandle::InvalidHandle;
+	// Check for valid parent
+	if (!Params.ParentHandle.IsValid() || !Params.NewClass)
+	{
+		return NewDataHandle;
+	}
+
+	UClass* NewClass = Params.NewClass;
+	UObject* Asset = Params.AssetOverride;
+	const FSubobjectDataHandle& ParentObjHandle = Params.ParentHandle;
+
+	if (NewClass->ClassWithin && NewClass->ClassWithin != UObject::StaticClass())
+	{
+		FailReason =  LOCTEXT("AddComponentFailed", "Cannot add components that have \"Within\" markup");
+		return NewDataHandle;
+	}
+	
+	FName TemplateVariableName;
+	const USCS_Node* SCSNode = Cast<const USCS_Node>(Asset);
+	UActorComponent* ComponentTemplate = (SCSNode ? ToRawPtr(SCSNode->ComponentTemplate) : Cast<UActorComponent>(Asset));
+
+	if (SCSNode)
+	{
+		TemplateVariableName = SCSNode->GetVariableName();
+		Asset = nullptr;
+	}
+	else if (ComponentTemplate)
+	{
+		Asset = nullptr;
+	}
+	
+	if(Params.BlueprintContext)
+	{
+		UBlueprint* Blueprint = Params.BlueprintContext;
+		check(Blueprint != nullptr && Blueprint->SimpleConstructionScript != nullptr);
+		Blueprint->Modify();
+		SaveSCSCurrentState(Blueprint->SimpleConstructionScript);
+		UActorComponent* NewComponent = nullptr;
+		
+		// Defer Blueprint class regeneration and tree updates until after we copy any object properties from a source template.
+        const bool bMarkBlueprintModified = false;
+		
+		FName NewVariableName;
+		if (ComponentTemplate)
+		{
+			if (!TemplateVariableName.IsNone())
+			{
+				NewVariableName = TemplateVariableName;
+			}
+			else
+			{
+				FString TemplateName = ComponentTemplate->GetName();
+				NewVariableName = (TemplateName.EndsWith(USimpleConstructionScript::ComponentTemplateNameSuffix) 
+                                    ? FName(*TemplateName.LeftChop(USimpleConstructionScript::ComponentTemplateNameSuffix.Len()))
+                                    : ComponentTemplate->GetFName());
+			}
+		}
+		else if (Asset)
+		{
+			NewVariableName = *FComponentEditorUtils::GenerateValidVariableNameFromAsset(Asset, nullptr);
+		}
+
+		USCS_Node* NewSCSNode = Blueprint->SimpleConstructionScript->CreateNode(NewClass, NewVariableName);
+		NewSCSNode->Modify();
+		NewComponent = NewSCSNode->ComponentTemplate;
+		// stuff from AddNewNode
+		// Assign this new node to an override asset if there is one
+		if(Asset)
+		{
+			FComponentAssetBrokerage::AssignAssetToComponent(NewComponent, Asset);
+		}
+
+		FSubobjectDataHandle TargetAttachmentHandle = FindParentForNewSubobject(NewComponent, ParentObjHandle);
+		FSubobjectData* TargetAttachment = TargetAttachmentHandle.GetData();
+		check(TargetAttachment);
+		
+		// Create a new subobject data set with this component. Use the SCS node here and the subobject data
+		// will correctly associate the component template
+		NewDataHandle = FactoryCreateSubobjectDataWithParent(NewSCSNode, TargetAttachment->GetHandle());
+
+		// Actually add this new node to the SimpleConstructionScript on the blueprint
+		AttachSubobject(TargetAttachment->GetHandle(), NewDataHandle);
+		
+		// Potentially adjust variable names for any child blueprints
+        const FName VariableName = NewSCSNode->GetVariableName();
+        if(VariableName != NAME_None)
+        {
+        	FBlueprintEditorUtils::ValidateBlueprintChildVariables(Blueprint, VariableName);
+        }
+		
+		if (ComponentTemplate)
+		{
+			//Serialize object properties using write/read operations.
+			TArray<uint8> SavedProperties;
+			FObjectWriter Writer(ComponentTemplate, SavedProperties);
+			FObjectReader(NewComponent, SavedProperties);
+			NewComponent->UpdateComponentToWorld();
+		}
+
+		if (Params.bConformTransformToParent)
+		{
+			if (USceneComponent* AsSceneComp = Cast<USceneComponent>(NewComponent))
+			{
+				// This problem because it is using the generated bp class as a parent
+				if (const USceneComponent* ParentSceneComp = CastChecked<USceneComponent>(TargetAttachment->GetComponentTemplate(), ECastCheckedType::NullAllowed))
+				{
+					ConformTransformRelativeToParent(AsSceneComp, ParentSceneComp);
+				}
+			}
+		}
+
+		// Wait until here to mark as structurally modified because we don't want any RerunConstructionScript() calls
+		// to happen until AFTER we've serialized properties from the source object.
+		if (!Params.bSkipMarkBlueprintModified)
+		{
+			FBlueprintEditorUtils::MarkBlueprintAsStructurallyModified(Blueprint);
+		}
+	}
+	else	// Not in a BP context
+	{
+		FSubobjectData* ParentObjData = ParentObjHandle.GetData();
+
+		check(ParentObjData);
+		
+		// If this is a component template, then we can just duplicate it
+		if(ComponentTemplate)
+		{
+			UActorComponent* NewComponent = FComponentEditorUtils::DuplicateComponent(ComponentTemplate);
+			// Create a new subobject data set with this component
+			NewDataHandle = FactoryCreateSubobjectDataWithParent(NewComponent, ParentObjHandle);
+		}
+		else if(AActor* ActorInstance = ParentObjData->GetMutableActorContext())
+		{
+			ActorInstance->Modify();
+
+			// Create an appropriate name for the new component
+			FName NewComponentName = NAME_None;
+			if (Asset)
+			{
+				NewComponentName = *FComponentEditorUtils::GenerateValidVariableNameFromAsset(Asset, ActorInstance);
+			}
+			else
+			{
+				NewComponentName = *FComponentEditorUtils::GenerateValidVariableName(NewClass, ActorInstance);
+			}
+			
+			// Get the set of owned components that exists prior to instancing the new component.
+			TInlineComponentArray<UActorComponent*> PreInstanceComponents;
+			ActorInstance->GetComponents(PreInstanceComponents);
+
+			// Construct the new component and attach as needed
+			UActorComponent* NewInstanceComponent = NewObject<UActorComponent>(ActorInstance, NewClass, NewComponentName, RF_Transactional);
+			
+			// Do Scene Attachment if this new Component is a USceneComponent
+			if (USceneComponent* NewSceneComponent = Cast<USceneComponent>(NewInstanceComponent))
+			{
+				if(ParentObjData->IsDefaultSceneRoot())
+				{
+					ActorInstance->SetRootComponent(NewSceneComponent);
+				}
+				else
+				{
+					USceneComponent* AttachTo = Cast<USceneComponent>(ParentObjData->GetMutableComponentTemplate());
+					if (AttachTo == nullptr)
+					{
+						AttachTo = ActorInstance->GetRootComponent();
+					}
+					check(AttachTo != nullptr);
+
+					// Make sure that the mobility of the new scene component is such that we can attach it
+					if (AttachTo->Mobility == EComponentMobility::Movable)
+					{
+						NewSceneComponent->Mobility = EComponentMobility::Movable;
+					}
+					else if (AttachTo->Mobility == EComponentMobility::Stationary && NewSceneComponent->Mobility == EComponentMobility::Static)
+					{
+						NewSceneComponent->Mobility = EComponentMobility::Stationary;
+					}
+
+					NewSceneComponent->AttachToComponent(AttachTo, FAttachmentTransformRules::SnapToTargetNotIncludingScale);
+				}
+			}
+
+			// If the component was created from/for a particular asset, assign it now
+			if (Asset)
+			{
+				FComponentAssetBrokerage::AssignAssetToComponent(NewInstanceComponent, Asset);
+			}
+
+			// Add to SerializedComponents array so it gets saved
+			ActorInstance->AddInstanceComponent(NewInstanceComponent);
+			NewInstanceComponent->OnComponentCreated();
+			NewInstanceComponent->RegisterComponent();
+
+			// Register any new components that may have been created during construction of the instanced component, but were not explicitly registered.
+			TInlineComponentArray<UActorComponent*> PostInstanceComponents;
+			ActorInstance->GetComponents(PostInstanceComponents);
+			for (UActorComponent* ActorComponent : PostInstanceComponents)
+			{
+				if (!ActorComponent->IsRegistered() && ActorComponent->bAutoRegister && !ActorComponent->IsPendingKill() && !PreInstanceComponents.Contains(ActorComponent))
+				{
+					ActorComponent->RegisterComponent();
+				}
+			}
+
+			// Rerun construction scripts
+			ActorInstance->RerunConstructionScripts();
+
+			// If the running the construction script destroyed the new node, don't create an entry for it
+			if (!NewInstanceComponent->IsPendingKill())
+			{
+				// Create a new subobject data set with this component
+				NewDataHandle = FactoryCreateSubobjectDataWithParent(NewInstanceComponent, ParentObjHandle);
+			}
+		}
+		else
+		{
+			FailReason =  LOCTEXT("AddComponentFailed_Inherited", "Cannot add components within an Inherited heirarchy");
+		}
+	}
+	
+	return NewDataHandle;
+}
+
+int32 USubobjectDataSubsystem::DeleteSubobjects(const FSubobjectDataHandle& ContextHandle, const TArray<FSubobjectDataHandle>& SubobjectsToDelete, FSubobjectDataHandle& OutComponentToSelect, UBlueprint* BPContext/* = nullptr*/)
+{
+	int32 NumDeletedSubobjects = 0;
+
+	if(!ContextHandle.IsValid() || SubobjectsToDelete.Num() == 0)
+	{
+		return NumDeletedSubobjects;
+	}
+	const FSubobjectData* ContextData = ContextHandle.GetData();
+	UObject* ContextObj = ContextData->GetMutableObject();
+	check(ContextObj);
+
+	ContextObj->Modify();
+	
+	if (BPContext)
+	{
+		for (const FSubobjectDataHandle& Handle : SubobjectsToDelete)
+		{
+			if (Handle.IsValid())
+			{
+				if (const FSubobjectData* Data = Handle.GetData())
+				{
+					USCS_Node* SCS_Node = Data->GetSCSNode();
+					
+					if (SCS_Node)
+					{
+						USimpleConstructionScript* SCS = SCS_Node->GetSCS();
+						check(SCS != nullptr && BPContext == SCS->GetBlueprint());
+						BPContext->Modify();
+						SaveSCSCurrentState(SCS);
+					
+						FBlueprintEditorUtils::RemoveVariableNodes(BPContext, Data->GetVariableName());
+
+						// If there are any Bound Component events for this property then give them compiler errors
+						TArray<UK2Node_ComponentBoundEvent*> EventNodes;
+						FKismetEditorUtilities::FindAllBoundEventsForComponent(BPContext, SCS_Node->GetVariableName(), EventNodes);
+						if (EventNodes.Num() > 0)
+						{
+							// Find any dynamic delegate nodes and give a compiler error for each that is problematic
+							FCompilerResultsLog LogResults;
+							FMessageLog MessageLog("BlueprintLog");
+
+							// Add a compiler error for each bound event node
+							for (UK2Node_ComponentBoundEvent* Node : EventNodes)
+							{
+								LogResults.Error(*LOCTEXT("RemoveBoundEvent_Error", "The component that @@ was bound to has been deleted! This node is no longer valid").ToString(), Node);
+							}
+
+							// Notify the user that these nodes are no longer valid
+							MessageLog.NewPage(LOCTEXT("RemoveBoundEvent_Error_Label", "Removed Owner of Component Bound Event"));
+							MessageLog.AddMessages(LogResults.Messages);
+							MessageLog.Notify(LOCTEXT("RemoveBoundEvent_Error_Msg", "Removed Owner of Component Bound Event"));
+
+							// Focus on the first node that we found
+							FKismetEditorUtilities::BringKismetToFocusAttentionOnObject(EventNodes[0]);
+						}
+
+						// Remove node from SCS tree
+						SCS->RemoveNodeAndPromoteChildren(SCS_Node);
+						++NumDeletedSubobjects;
+						
+						// Clear the delegate
+						SCS_Node->SetOnNameChanged(FSCSNodeNameChanged());
+
+						// on removal, since we don't move the template from the GeneratedClass (which we shouldn't, as it would create a 
+						// discrepancy with existing instances), we rename it instead so that we can re-use the name without having to compile  
+						// (we still have a problem if they attempt to name it to what ever we choose here, but that is unlikely)
+						// note: skip this for the default scene root; we don't actually destroy that node when it's removed, so we don't need the template to be renamed.
+						if (!Data->IsDefaultSceneRoot() && SCS_Node->ComponentTemplate != nullptr)
+						{
+							const FName TemplateName = SCS_Node->ComponentTemplate->GetFName();
+							const FString RemovedName = SCS_Node->GetVariableName().ToString() + TEXT("_REMOVED_") + FGuid::NewGuid().ToString();
+
+							SCS_Node->ComponentTemplate->Modify();
+							SCS_Node->ComponentTemplate->Rename(*RemovedName, /*NewOuter =*/nullptr, REN_DontCreateRedirectors);
+
+							TArray<UObject*> ArchetypeInstances;
+							auto DestroyArchetypeInstances = [&ArchetypeInstances, &RemovedName](UActorComponent* ComponentTemplate)
+							{
+								ComponentTemplate->GetArchetypeInstances(ArchetypeInstances);
+								for (UObject* ArchetypeInstance : ArchetypeInstances)
+								{
+									if (!ArchetypeInstance->HasAllFlags(RF_ArchetypeObject | RF_InheritableComponentTemplate))
+									{
+										CastChecked<UActorComponent>(ArchetypeInstance)->DestroyComponent();
+										ArchetypeInstance->Rename(*RemovedName, nullptr, REN_DontCreateRedirectors);
+									}
+								}
+							};
+
+							DestroyArchetypeInstances(SCS_Node->ComponentTemplate);
+
+							if (BPContext)
+							{
+								// Children need to have their inherited component template instance of the component renamed out of the way as well
+								TArray<UClass*> ChildrenOfClass;
+								GetDerivedClasses(BPContext->GeneratedClass, ChildrenOfClass);
+
+								for (UClass* ChildClass : ChildrenOfClass)
+								{
+									UBlueprintGeneratedClass* BPChildClass = CastChecked<UBlueprintGeneratedClass>(ChildClass);
+
+									if (UActorComponent* Component = (UActorComponent*)FindObjectWithOuter(BPChildClass, UActorComponent::StaticClass(), TemplateName))
+									{
+										Component->Modify();
+										Component->Rename(*RemovedName, /*NewOuter =*/nullptr, REN_DontCreateRedirectors);
+
+										DestroyArchetypeInstances(Component);
+									}
+								}
+							}
+						}
+					}
+				}
+			}
+		}
+
+		// Will call UpdateTree as part of OnBlueprintChanged handling
+		FBlueprintEditorUtils::MarkBlueprintAsStructurallyModified(BPContext);
+	}
+	else	// This is an actor instance
+	{
+		TArray<UActorComponent*> ComponentsToDelete;
+		for(const FSubobjectDataHandle& Handle : SubobjectsToDelete)
+		{
+			if(Handle.IsValid())
+			{
+				if(const FSubobjectData* Data = Handle.GetData())
+				{
+					ComponentsToDelete.Add(Data->GetMutableComponentTemplate());		
+				}		
+			}
+		}
+		UActorComponent* ActorComponentToSelect = nullptr;
+		NumDeletedSubobjects = FComponentEditorUtils::DeleteComponents(ComponentsToDelete, ActorComponentToSelect);
+
+		// Find the handle that matches with this actor component given our current context
+	}
+
+	return NumDeletedSubobjects;
+}
+
+int32 USubobjectDataSubsystem::DeleteSubobjects(const FSubobjectDataHandle& ContextHandle, const TArray<FSubobjectDataHandle>& SubobjectsToDelete, UBlueprint* BPContext /*= nullptr*/)
+{
+	FSubobjectDataHandle Dummy;
+	return DeleteSubobjects(ContextHandle, SubobjectsToDelete, Dummy, BPContext);
+}
+
+int32 USubobjectDataSubsystem::DeleteSubobject(const FSubobjectDataHandle& ContextHandle, const FSubobjectDataHandle& SubobjectToDelete, UBlueprint* BPContext)
+{
+	TArray<FSubobjectDataHandle> Handles { SubobjectToDelete };
+	return DeleteSubobjects(ContextHandle, Handles, BPContext);
+}
+
+bool USubobjectDataSubsystem::RenameSubobject(const FSubobjectDataHandle& Handle, const FText& InNewName)
+{
+	FText OutErrorMessage;
+	if(!IsValidRename(Handle, InNewName, OutErrorMessage))
+	{
+		return false;	
+	}
+
+	const FSubobjectData* Data = Handle.GetData();
+	if(!Data)
+	{
+		return false;
+	}
+
+	// For root actor nodes
+	if (AActor* Actor = Data->GetMutableObject<AActor>())
+	{
+		if (Actor->IsActorLabelEditable() && !InNewName.ToString().Equals(Actor->GetActorLabel(), ESearchCase::CaseSensitive))
+		{
+			const FScopedTransaction Transaction(LOCTEXT("SCSEditorRenameActorTransaction", "Rename Actor"));
+			FActorLabelUtilities::RenameExistingActor(Actor, InNewName.ToString());
+			return true;
+		}
+	}
+	
+	// For instanced components
+	if(UActorComponent* ComponentInstance = Data->GetMutableComponentTemplate())
+	{
+		if(Data->IsInstancedComponent())
+		{
+			ERenameFlags RenameFlags = REN_DontCreateRedirectors;
+	
+			// name collision could occur due to e.g. our archetype being updated and causing a conflict with our ComponentInstance:
+			FString NewNameAsString = InNewName.ToString();
+			if(StaticFindObject(UObject::StaticClass(), ComponentInstance->GetOuter(), *NewNameAsString) == nullptr)
+			{
+				ComponentInstance->Rename(*NewNameAsString, nullptr, RenameFlags);
+			}
+			return true;
+		}
+		else if(UBlueprint* BP = Data->GetBlueprint())
+		{
+			FName ValidatedNewName;
+			FString DesiredName = InNewName.ToString();
+			
+			// Is this desired name the same as what is already there? If so then don't bother
+			USCS_Node* SCSNode = Data->GetSCSNode();
+			if(SCSNode && SCSNode->GetVariableName().ToString().Equals(DesiredName))
+			{
+				return true;
+			}
+			
+			if (FKismetNameValidator(BP).IsValid(DesiredName) == EValidatorResult::Ok)
+			{
+				ValidatedNewName = FName(DesiredName);
+			}
+			else
+			{
+				ValidatedNewName = FBlueprintEditorUtils::FindUniqueKismetName(BP, DesiredName);
+			}
+			
+			FBlueprintEditorUtils::RenameComponentMemberVariable(BP, SCSNode, ValidatedNewName);
+			return true;
+		}
+	}
+	
+	return false;
+}
+
+bool USubobjectDataSubsystem::ReparentSubobject(const FReparentSubobjectParams& Params, const FSubobjectDataHandle& ToReparentHandle)
+{
+	TArray<FSubobjectDataHandle> Handles;
+	Handles.Add(ToReparentHandle);
+	
+	return ReparentSubobjects(Params, Handles);
+}
+
+bool USubobjectDataSubsystem::MakeNewSceneRoot(const FSubobjectDataHandle& Context, const FSubobjectDataHandle& DroppedNewSceneRootHandle, UBlueprint* Blueprint)
+{
+	if(!ensure(Context.IsValid()) || !ensure(DroppedNewSceneRootHandle.IsValid()))
+	{
+		UE_LOG(LogSubobjectSubsystem, Warning, TEXT("Failed to make new scene root: Invalid context or scene root handle!"));
+		return false;
+	}
+
+	FSubobjectData* DroppedNewSceneRootData = DroppedNewSceneRootHandle.GetData();
+	
+	FSubobjectDataHandle StartingRootHandle = FindSceneRootForSubobject(Context);
+	FSubobjectData* StartingRootData = StartingRootHandle.GetData();
+	// Remember whether or not we're replacing the default scene root
+	bool bWasDefaultSceneRoot = StartingRootData && StartingRootData->IsDefaultSceneRoot();
+	
+	FSubobjectDataHandle OldSceneRoot = FSubobjectDataHandle::InvalidHandle;
+	
+	if(Blueprint)
+	{
+		check(Blueprint && Blueprint->SimpleConstructionScript);
+
+		// Clone the component if it's being dropped into a different subobject tree
+		if(DroppedNewSceneRootData->GetBlueprint() != Blueprint)
+		{
+			UActorComponent* ComponentTemplate = DroppedNewSceneRootData->GetMutableComponentTemplate();
+			check(ComponentTemplate);
+			FAddNewSubobjectParams Params;
+			Params.NewClass = ComponentTemplate->GetClass();
+			Params.BlueprintContext = Blueprint;
+			Params.AssetOverride = nullptr;
+			Params.ParentHandle = Context;
+			FText FailReason;
+			
+			// Note: This will mark the Blueprint as structurally modified
+			FSubobjectDataHandle ClonedHandle = AddNewSubobject(Params, FailReason);
+			check(ClonedHandle.IsValid());
+			UActorComponent* ClonedComponent = ClonedHandle.GetData()->GetMutableComponentTemplate();
+			check(ClonedComponent);
+
+			// Serialize object properties using write/read operations.
+			TArray<uint8> SavedProperties;
+			FObjectWriter Writer(ComponentTemplate, SavedProperties);
+			FObjectReader(ClonedComponent, SavedProperties);
+
+			DroppedNewSceneRootData = ClonedHandle.GetData();
+			check(DroppedNewSceneRootData->IsValid());
+		}
+
+		if(DroppedNewSceneRootData->GetParentHandle().IsValid() && DroppedNewSceneRootData->GetBlueprint() == Blueprint)
+		{
+			// If the associated component template is a scene component, reset its transform since it will now become the root
+			if(USceneComponent* SceneComponentTemplate = Cast<USceneComponent>(DroppedNewSceneRootData->GetMutableComponentTemplate()))
+			{
+				// Save current state
+				SceneComponentTemplate->Modify();
+
+				// Reset the attach socket name
+				SceneComponentTemplate->SetupAttachment(SceneComponentTemplate->GetAttachParent(), NAME_None);
+				
+				if(USCS_Node* SCS_Node = DroppedNewSceneRootData->GetSCSNode())
+				{
+					SCS_Node->Modify();
+					SCS_Node->AttachToName = NAME_None;
+				}
+
+				// Cache the current relative location and rotation values (for propagation)
+				const FVector OldRelativeLocation = SceneComponentTemplate->GetRelativeLocation();
+				const FRotator OldRelativeRotation = SceneComponentTemplate->GetRelativeRotation();
+
+				// Reset the relative transform (location and rotation only; scale is preserved)
+				SceneComponentTemplate->SetRelativeLocation(FVector::ZeroVector);
+				SceneComponentTemplate->SetRelativeRotation(FRotator::ZeroRotator);
+
+				// Propagate the root change & detachment to any instances of the template (done within the context of the current transaction)
+				TArray<UObject*> ArchetypeInstances;
+				SceneComponentTemplate->GetArchetypeInstances(ArchetypeInstances);
+				FDetachmentTransformRules DetachmentTransformRules(EDetachmentRule::KeepWorld, EDetachmentRule::KeepWorld, EDetachmentRule::KeepRelative, true);
+				for (int32 InstanceIndex = 0; InstanceIndex < ArchetypeInstances.Num(); ++InstanceIndex)
+				{
+					USceneComponent* SceneComponentInstance = Cast<USceneComponent>(ArchetypeInstances[InstanceIndex]);
+					if (SceneComponentInstance != nullptr)
+					{
+						// Detach from root (keeping world transform, except for scale)
+						SceneComponentInstance->DetachFromComponent(DetachmentTransformRules);
+
+						// Propagate the default relative location & rotation reset from the template to the instance
+						FComponentEditorUtils::ApplyDefaultValueChange(SceneComponentInstance, SceneComponentInstance->GetRelativeLocation_DirectMutable(), OldRelativeLocation, SceneComponentTemplate->GetRelativeLocation());
+						FComponentEditorUtils::ApplyDefaultValueChange(SceneComponentInstance, SceneComponentInstance->GetRelativeRotation_DirectMutable(), OldRelativeRotation, SceneComponentTemplate->GetRelativeRotation());
+
+						// Must also reset the root component here, so that RerunConstructionScripts() will cache the correct root component instance data
+						AActor* Owner = SceneComponentInstance->GetOwner();
+						if (Owner)
+						{
+							Owner->Modify();
+							Owner->SetRootComponent(SceneComponentInstance);
+						}
+					}
+				}
+			}
+			
+			// Remove the dropped node from its existing parent
+			DetachSubobject(DroppedNewSceneRootData->GetParentHandle(), DroppedNewSceneRootData->GetHandle());
+		}
+
+		check(StartingRootData && (bWasDefaultSceneRoot || StartingRootData->CanReparent()));
+
+		// Remove the current scene root node from the SCS context
+		Blueprint->SimpleConstructionScript->RemoveNode(StartingRootData->GetSCSNode(), /*bValidateSceneRootNodes=*/false);
+
+		// Save old root node
+		OldSceneRoot = StartingRootData->GetHandle();
+
+		// Add dropped node to the SCS context
+		Blueprint->SimpleConstructionScript->AddNode(DroppedNewSceneRootData->GetSCSNode());
+		
+		// Remove or re-parent the old root
+		if (OldSceneRoot.IsValid())
+		{
+			check(DroppedNewSceneRootData->CanReparent());
+
+			// Set old root as child of new root
+			AttachSubobject(DroppedNewSceneRootData->GetHandle(), OldSceneRoot);
+
+			if (bWasDefaultSceneRoot)
+			{
+				DeleteSubobject(Context, OldSceneRoot, Blueprint);
+			}
+		}
+	}
+	else
+	{
+		// Remove the dropped node from its existing parent
+		if(DroppedNewSceneRootData->HasParent())
+		{
+			DetachSubobject(DroppedNewSceneRootData->GetParentHandle(), DroppedNewSceneRootData->GetHandle());
+		}
+		
+		// Save old root node
+		OldSceneRoot = StartingRootHandle;
+
+		if(OldSceneRoot.IsValid())
+		{
+			// If the thing we replaced was the default scene root, then just delete it
+			if(bWasDefaultSceneRoot)
+			{
+				DeleteSubobject(Context, OldSceneRoot, Blueprint);
+				AActor* ActorContext = Context.GetData()->GetMutableActorContext();
+				ActorContext->SetRootComponent(CastChecked<USceneComponent>(DroppedNewSceneRootData->GetMutableComponentTemplate()));
+			}
+			// Otherwise, reparent it to the new scene root
+			else
+			{
+				FReparentSubobjectParams Params;
+				Params.BlueprintContext = nullptr;
+				Params.ActorPreviewContext = nullptr;
+				Params.NewParentHandle = DroppedNewSceneRootHandle;
+				
+				ReparentSubobject(Params, OldSceneRoot);
+			}
+		}
+	}
+
+	return true;
+}
+
+bool USubobjectDataSubsystem::ReparentSubobjects(const FReparentSubobjectParams& Params, const TArray<FSubobjectDataHandle>& HandlesToMove)
+{
+	if(!Params.NewParentHandle.IsValid())
+	{
+		UE_LOG(LogSubobjectSubsystem, Warning, TEXT("Failed to reparent: Invalid parent handle when reparenting!"));
+		return false;
+	}
+	
+	FSubobjectData* NewParentData = Params.NewParentHandle.GetData();
+	check(NewParentData);
+	
+	if (Params.BlueprintContext)
+	{
+		if (!Params.ActorPreviewContext)
+		{
+			UE_LOG(LogSubobjectSubsystem, Warning, TEXT("Failed to reparent: In a blueprint context there must be an actor preview!"));
+			return false;
+		}
+
+		for (const FSubobjectDataHandle& HandleToMove : HandlesToMove)
+		{
+			FSubobjectData* DroppedData = HandleToMove.GetSharedDataPtr().Get();
+			if(DroppedData->GetBlueprint() != Params.BlueprintContext)
+			{
+				UActorComponent* ComponentTemplate = DroppedData->GetMutableComponentTemplate();
+				check(ComponentTemplate);
+	
+				FAddNewSubobjectParams AddNewParms;
+				AddNewParms.BlueprintContext = Params.BlueprintContext;
+				AddNewParms.NewClass = ComponentTemplate->GetClass();
+				AddNewParms.ParentHandle = NewParentData->GetHandle();
+				
+				// Note: This will mark the Blueprint as structurally modified
+				FText FailReason;
+				FSubobjectDataHandle ClonedSubobject = AddNewSubobject(AddNewParms, FailReason);
+				check(ClonedSubobject.IsValid());
+				
+				FSubobjectData* ClonedData = ClonedSubobject.GetSharedDataPtr().Get();
+				UActorComponent* ClonedComponent = ClonedData->GetMutableComponentTemplate();
+	
+				// Serialize object properties using write/read operations.
+				TArray<uint8> SavedProperties;
+				FObjectWriter Writer(ComponentTemplate, SavedProperties);
+				FObjectReader(ClonedComponent, SavedProperties);			
+			}
+			else
+			{
+				USceneComponent* SceneComponentTemplate = Cast<USceneComponent>(DroppedData->GetMutableComponentTemplate());
+				// Cache current default values for propagation
+				FVector OldRelativeLocation, OldRelativeScale3D;
+				FRotator OldRelativeRotation;
+				if(SceneComponentTemplate)
+				{
+					OldRelativeLocation = SceneComponentTemplate->GetRelativeLocation();
+					OldRelativeRotation = SceneComponentTemplate->GetRelativeRotation();
+					OldRelativeScale3D = SceneComponentTemplate->GetRelativeScale3D();
+				}
+
+				FSubobjectDataHandle OldParentHandle = DroppedData->GetParentHandle();
+				FSubobjectData* OldParentData = OldParentHandle.IsValid() ? OldParentHandle.GetSharedDataPtr().Get() : nullptr;
+
+				if(OldParentData)
+				{
+					DetachSubobject(OldParentHandle, DroppedData->GetHandle());
+					// If the associated component template is a scene component, maintain its preview world position
+					if(SceneComponentTemplate)
+					{
+						// Save current state
+						SceneComponentTemplate->Modify();
+
+						// Reset the attach socket name
+						SceneComponentTemplate->SetupAttachment(SceneComponentTemplate->GetAttachParent(), NAME_None);
+						USCS_Node* SCS_Node = DroppedData->GetSCSNode();
+						if(SCS_Node)
+						{
+							SCS_Node->Modify();
+							SCS_Node->AttachToName = NAME_None;
+						}
+
+						// Attempt to locate a matching registered instance of the component template in the Actor context that's being edited
+						const USceneComponent* InstancedSceneComponent = Cast<USceneComponent>(DroppedData->FindMutableComponentInstanceInActor(Params.ActorPreviewContext));
+						if(InstancedSceneComponent && InstancedSceneComponent->IsRegistered())
+						{
+							// If we find a match, save off the world position
+							const FTransform& ComponentToWorld = InstancedSceneComponent->GetComponentToWorld();
+							SceneComponentTemplate->SetRelativeTransform_Direct(ComponentToWorld);
+						}
+					}
+				}
+
+				// Attach the dropped node to the given node
+				AttachSubobject(NewParentData->GetHandle(), DroppedData->GetHandle());
+
+				// Attempt to locate a matching instance of the parent component template in the Actor context that's being edited
+				USceneComponent* ParentSceneComponent = Cast<USceneComponent>(DroppedData->FindMutableComponentInstanceInActor(Params.ActorPreviewContext));
+				if(SceneComponentTemplate && ParentSceneComponent && ParentSceneComponent->IsRegistered())
+				{
+					ConformTransformRelativeToParent(SceneComponentTemplate, ParentSceneComponent);
+				}
+
+				// Propagate any default value changes out to all instances of the template. If we didn't do this, then instances could incorrectly override the new default value with the old default value when construction scripts are re-run.
+				if(SceneComponentTemplate)
+				{
+					TArray<UObject*> InstancedSceneComponents;
+					SceneComponentTemplate->GetArchetypeInstances(InstancedSceneComponents);
+					for(int32 InstanceIndex = 0; InstanceIndex < InstancedSceneComponents.Num(); ++InstanceIndex)
+					{
+						USceneComponent* InstancedSceneComponent = Cast<USceneComponent>(InstancedSceneComponents[InstanceIndex]);
+						if(InstancedSceneComponent != nullptr)
+						{
+							FComponentEditorUtils::ApplyDefaultValueChange(InstancedSceneComponent, InstancedSceneComponent->GetRelativeLocation_DirectMutable(), OldRelativeLocation, SceneComponentTemplate->GetRelativeLocation());
+							FComponentEditorUtils::ApplyDefaultValueChange(InstancedSceneComponent, InstancedSceneComponent->GetRelativeRotation_DirectMutable(), OldRelativeRotation, SceneComponentTemplate->GetRelativeRotation());
+							FComponentEditorUtils::ApplyDefaultValueChange(InstancedSceneComponent, InstancedSceneComponent->GetRelativeScale3D_DirectMutable(),  OldRelativeScale3D,  SceneComponentTemplate->GetRelativeScale3D());
+						}
+					}
+				}
+			}
+		}
+
+		FBlueprintEditorUtils::PostEditChangeBlueprintActors(Params.BlueprintContext, true);
+		return true;
+	}
+	else
+	{
+		for (const FSubobjectDataHandle& HandleToMove : HandlesToMove)
+		{
+			// Remove this handle from it's parent if it has one
+			FSubobjectData* DataToMove = HandleToMove.GetData();
+
+			if (DataToMove->HasParent())
+			{
+				DetachSubobject(DataToMove->GetParentHandle(), HandleToMove);
+			}
+
+			AttachSubobject(Params.NewParentHandle, HandleToMove);
+		}
+
+		// Rerun construction scripts on the base actor 
+		if (AActor* ActorInstance = NewParentData->GetMutableActorContext())
+		{
+			ActorInstance->RerunConstructionScripts();
+		}
+	}
+	
+	return true;
+}
+
+bool USubobjectDataSubsystem::DetachSubobject(const FSubobjectDataHandle& OwnerHandle, const FSubobjectDataHandle& ChildToRemove)
+{
+	FSubobjectData* OwnerData = OwnerHandle.GetData();
+	FSubobjectData* ChildToRemoveData = ChildToRemove.GetData();
+	if(!OwnerData || !ChildToRemoveData)
+	{
+		return false;
+	}
+
+	// Remove this subobject handle from the parent
+	OwnerData->RemoveChildHandleOnly(ChildToRemoveData->GetHandle());
+	ChildToRemoveData->ClearParentHandle();
+	// if its an instance component, call detact from compont
+	if(ChildToRemoveData->IsInstancedComponent())
+	{
+		USceneComponent* ChildInstance = Cast<USceneComponent>(ChildToRemoveData->GetMutableComponentTemplate());
+		if (ensure(ChildInstance))
+		{
+			// Handle detachment at the instance level
+			ChildInstance->DetachFromComponent(FDetachmentTransformRules::KeepWorldTransform);
+		}
+		return true;
+	}
+
+	// Bypass removal logic if we're part of a child actor subtree
+	if (ChildToRemoveData->IsChildActor())
+	{
+		return true;
+	}
+	
+	if(USCS_Node* SCS_ChildNode = ChildToRemoveData->GetSCSNode())
+	{
+		USimpleConstructionScript* SCS = SCS_ChildNode->GetSCS();
+		
+		if (SCS != nullptr)
+		{
+			SCS->RemoveNode(SCS_ChildNode);
+		}
+	}
+	
+	return true;
+}
+
+bool USubobjectDataSubsystem::AttachSubobject(const FSubobjectDataHandle& OwnerHandle, const FSubobjectDataHandle& ChildToAddHandle)
+{
+	FSubobjectData* OwnerData = OwnerHandle.GetData();
+	FSubobjectData* ChildToAddData = ChildToAddHandle.GetData();
+
+	if(!OwnerData || !ChildToAddData)
+	{
+		return false;
+	}
+
+	// If the given child has a parent already, remove it from that first
+	if(ChildToAddData->HasParent())
+	{
+		DetachSubobject(ChildToAddData->GetParentHandle(), ChildToAddHandle);
+	}
+
+	check(!ChildToAddData->HasParent());
+	
+	// Reparent it to the new data
+	OwnerData->AddChildHandleOnly(ChildToAddHandle);
+	ChildToAddData->SetParentHandle(OwnerHandle);
+	
+	if(ChildToAddData->IsComponent())
+	{
+		USCS_Node* SCS_Node = OwnerData->GetSCSNode();
+		const UActorComponent* ComponentTemplate = OwnerData->GetObject<UActorComponent>();
+		
+		// Add a child node to the SCS tree node if not already present
+		if(USCS_Node* SCS_ChildNode = ChildToAddData->GetSCSNode())
+		{
+			// Get the SCS instance that owns the child node
+			if(USimpleConstructionScript* SCS = SCS_ChildNode->GetSCS())
+			{
+				if(SCS_Node)
+				{
+					// If the parent and child are both owned by the same SCS instance
+					if (SCS_Node->GetSCS() == SCS)
+					{
+						// Add the child into the parent's list of children
+						if (!SCS_Node->GetChildNodes().Contains(SCS_ChildNode))
+						{
+							SCS_Node->AddChildNode(SCS_ChildNode);
+						}
+					}
+					else
+					{
+						// Adds the child to the SCS root set if not already present
+						SCS->AddNode(SCS_ChildNode);
+
+						// Set parameters to parent this node to the "inherited" SCS node
+						SCS_ChildNode->SetParent(SCS_Node);
+					}
+				}
+				else if(ComponentTemplate)
+				{
+					// Adds the child to the SCS root set if not already present
+					SCS->AddNode(SCS_ChildNode);
+					
+					// Set parameters to parent this node to the native component template
+					SCS_ChildNode->SetParent(Cast<const USceneComponent>(ComponentTemplate));
+				}
+				else
+				{
+					// Adds the child to the SCS root set if not already present
+					SCS->AddNode(SCS_ChildNode);
+				}
+			}
+		}
+		else if(OwnerData->IsInstancedComponent())
+		{
+			USceneComponent* ChildInstance = Cast<USceneComponent>(ChildToAddData->GetMutableComponentTemplate());
+			if (ensure(ChildInstance != nullptr))
+			{
+				USceneComponent* ParentInstance = Cast<USceneComponent>(OwnerData->GetMutableComponentTemplate());
+				if (ensure(ParentInstance != nullptr))
+				{
+					// Handle attachment at the instance level
+					if (ChildInstance->GetAttachParent() != ParentInstance)
+					{
+						AActor* Owner = ParentInstance->GetOwner();
+						if (Owner->GetRootComponent() == ChildInstance)
+						{
+							Owner->SetRootComponent(ParentInstance);
+						}
+						ChildInstance->AttachToComponent(ParentInstance, FAttachmentTransformRules::KeepWorldTransform);
+					}
+				}
+			}
+		}
+	}
+	
+	return true;
+}
+
+bool USubobjectDataSubsystem::IsValidRename(const FSubobjectDataHandle& Handle, const FText& InNewText, FText& OutErrorMessage) const
+{
+	const FSubobjectData* Data = Handle.GetData();
+	if(!Data)
+	{
+		return false;		
+	}
+	
+	const UBlueprint* Blueprint = Data->GetBlueprint();
+	const FString& NewTextStr = InNewText.ToString();
+
+	if (!NewTextStr.IsEmpty())
+	{
+		if (Data->GetVariableName().ToString() == NewTextStr)
+		{
+			return true;
+		}
+
+		if (const UActorComponent* ComponentInstance = Data->GetComponentTemplate())
+		{
+			AActor* ExistingNameSearchScope = ComponentInstance->GetOwner();
+			if ((ExistingNameSearchScope == nullptr) && (Blueprint != nullptr))
+			{
+				ExistingNameSearchScope = Cast<AActor>(Blueprint->GeneratedClass->GetDefaultObject());
+			}
+
+			if (!FComponentEditorUtils::IsValidVariableNameString(ComponentInstance, NewTextStr))
+			{
+				OutErrorMessage = LOCTEXT("RenameFailed_EngineReservedName", "This name is reserved for engine use.");
+				return false;
+			}
+			else if (NewTextStr.Len() > NAME_SIZE)
+			{
+				FFormatNamedArguments Arguments;
+				Arguments.Add(TEXT("CharCount"), NAME_SIZE);
+				OutErrorMessage = FText::Format(LOCTEXT("ComponentRenameFailed_TooLong", "Component name must be less than {CharCount} characters long."), Arguments);
+				return false;
+			}
+			else if (!FComponentEditorUtils::IsComponentNameAvailable(NewTextStr, ExistingNameSearchScope, ComponentInstance) 
+                    || !FComponentEditorUtils::IsComponentNameAvailable(NewTextStr, ComponentInstance->GetOuter(), ComponentInstance ))
+			{
+				OutErrorMessage = LOCTEXT("RenameFailed_ExistingName", "Another component already has the same name.");
+				return false;
+			}
+		}
+		else if(const AActor* ActorInstance = Data->GetObject<AActor>())
+		{
+			// #TODO_BH Validation of actor instance
+		}
+		else
+		{
+			OutErrorMessage = LOCTEXT("RenameFailed_InvalidComponentInstance", "This node is referencing an invalid component instance and cannot be renamed. Perhaps it was destroyed?");
+			return false;
+		}
+	}
+	
+	TSharedPtr<INameValidatorInterface> NameValidator;
+	if (Blueprint != nullptr)
+	{
+		NameValidator = MakeShareable(new FKismetNameValidator(Blueprint, Data->GetVariableName()));
+	}
+	else if(const UActorComponent* CompTemplate = Data->GetComponentTemplate())
+	{
+		NameValidator = MakeShareable(new FStringSetNameValidator(CompTemplate->GetName()));
+	}
+
+	if(NameValidator)
+	{
+		EValidatorResult ValidatorResult = NameValidator->IsValid(NewTextStr);
+		if (ValidatorResult == EValidatorResult::AlreadyInUse)
+		{
+			OutErrorMessage = FText::Format(LOCTEXT("RenameFailed_InUse", "{0} is in use by another variable or function!"), InNewText);
+		}
+		else if (ValidatorResult == EValidatorResult::EmptyName)
+		{
+			OutErrorMessage = LOCTEXT("RenameFailed_LeftBlank", "Names cannot be left blank!");
+		}
+		else if (ValidatorResult == EValidatorResult::TooLong)
+		{
+			OutErrorMessage = LOCTEXT("RenameFailed_NameTooLong", "Names must have fewer than 100 characters!");
+		}
+	}
+
+	if (OutErrorMessage.IsEmpty())
+	{
+		return true;
+	}
+
+	return false;
+}
+
+bool USubobjectDataSubsystem::CanCopySubobjects(const TArray<FSubobjectDataHandle>& Handles) const
+{
+	TArray<UActorComponent*> ComponentsToCopy;
+	
+	for (const FSubobjectDataHandle& Handle : Handles)
+    {
+    	if(FSubobjectData* Data = Handle.GetData())
+    	{
+    		if(!Data->CanCopy())
+    		{
+    			return false;
+    		}
+    		
+    		// Get the component template associated with the selected node
+    		if (UActorComponent* ComponentTemplate = Data->GetMutableComponentTemplate())
+    		{
+    			ComponentsToCopy.Add(ComponentTemplate);
+    		}
+    	}
+    }
+
+    // Verify that the components can be copied
+    return FComponentEditorUtils::CanCopyComponents(ComponentsToCopy);
+}
+
+void USubobjectDataSubsystem::CopySubobjects(const TArray<FSubobjectDataHandle>& Handles, UBlueprint* BpContext)
+{
+	if(!CanCopySubobjects(Handles))
+	{
+		return;
+	}
+
+	TArray<UActorComponent*> ComponentsToCopy;
+
+	for (const FSubobjectDataHandle& Handle : Handles)
+	{
+		if(FSubobjectData* Data = Handle.GetData())
+		{
+			ensureMsgf(Data->CanCopy(), TEXT("A non-copiable subobject has been allowed to copy!"));
+    		
+			// Get the component template associated with the selected node
+			if (UActorComponent* ComponentTemplate = Data->GetMutableComponentTemplate())
+			{
+				ComponentsToCopy.Add(ComponentTemplate);
+				if(BpContext && ComponentTemplate->CreationMethod != EComponentCreationMethod::UserConstructionScript)
+				{
+					// CopyComponents uses component attachment to maintain hierarchy, but the SCS templates are not
+                    // setup with a relationship to each other. Briefly setup the attachment between the templates being
+                    // copied so that the hierarchy is retained upon pasting
+                    if (USceneComponent* SceneTemplate = Cast<USceneComponent>(ComponentTemplate))
+                    {
+                    	if (FSubobjectData* SelectedParentNodePtr = Data->GetParentHandle().GetData())
+                    	{
+                    		if (USceneComponent* ParentSceneTemplate = Cast<USceneComponent>(SelectedParentNodePtr->GetMutableComponentTemplate()))
+                    		{
+                    			SceneTemplate->SetupAttachment(ParentSceneTemplate);
+                    		}
+                    	}
+                    }
+				}
+			}
+		}
+	}
+	
+	// Copy the components to the clipboard
+	FComponentEditorUtils::CopyComponents(ComponentsToCopy);
+	
+	if(BpContext)
+	{
+		for (UActorComponent* ComponentTemplate : ComponentsToCopy)
+		{
+			if (ComponentTemplate->CreationMethod != EComponentCreationMethod::UserConstructionScript)
+			{
+				if (USceneComponent* SceneTemplate = Cast<USceneComponent>(ComponentTemplate))
+				{
+					// clear back out any temporary attachments we set up for the copy
+					SceneTemplate->SetupAttachment(nullptr);
+				}
+			}
+		}
+	}
+}
+
+bool USubobjectDataSubsystem::CanPasteSubobjects(const FSubobjectDataHandle& RootHandle, UBlueprint* BPContext) const
+{
+	const FSubobjectDataHandle& SceneRootHandle = FindSceneRootForSubobject(RootHandle);
+	const FSubobjectData* RootData = SceneRootHandle.GetData();
+	const USceneComponent* SceneComponent = Cast<USceneComponent>(RootData->GetComponentTemplate());
+	if(const AActor* RootActor = RootData->GetObject<AActor>())
+	{
+		SceneComponent = RootActor->GetRootComponent();
+	}
+	
+	return
+		(BPContext && RootData->IsActor())  ||
+		(SceneComponent && FComponentEditorUtils::CanPasteComponents(SceneComponent, RootData->IsDefaultSceneRoot(), true));
+}
+
+void USubobjectDataSubsystem::PasteSubobjects(const FSubobjectDataHandle& PasteToContext, const TArray<FSubobjectDataHandle>& NewParentHandles, UBlueprint* Blueprint, TArray<FSubobjectDataHandle>& OutPastedHandles)
+{
+	if(!PasteToContext.IsValid() || NewParentHandles.IsEmpty())
+	{
+		return;
+	}
+	
+	FSubobjectData* PasteToContextData = PasteToContext.GetSharedDataPtr().Get();	
+	
+	const FScopedTransaction Transaction(LOCTEXT("PasteComponents", "Paste Component(s)"));
+	if(Blueprint)
+	{
+		// Get the components to paste from the clipboard
+		TMap<FName, FName> ParentMap;
+		TMap<FName, UActorComponent*> NewObjectMap;
+		FComponentEditorUtils::GetComponentsFromClipboard(ParentMap, NewObjectMap, true);
+
+		check(Blueprint->SimpleConstructionScript != nullptr);
+		
+		Blueprint->Modify();
+		SaveSCSCurrentState(Blueprint->SimpleConstructionScript);
+		
+		// Create a new tree node for each new (pasted) component
+        FSubobjectDataHandle FirstNode;
+        TMap<FName, FSubobjectDataHandle> NewNodeMap;
+
+		for (const TPair<FName, UActorComponent*>& NewObjectPair : NewObjectMap)
+		{
+			// Get the component object instance
+			UActorComponent* NewActorComponent = NewObjectPair.Value;
+			check(NewActorComponent);
+			
+			// Create a new SCS node to contain the new component and add it to the tree
+			USCS_Node* NewSCSNode = Blueprint->SimpleConstructionScript->CreateNodeAndRenameComponent(NewActorComponent);
+			NewActorComponent = NewSCSNode ? NewSCSNode->ComponentTemplate : nullptr;
+
+			FSubobjectDataHandle TargetParentHandle = FindParentForNewSubobject(NewActorComponent, PasteToContext);
+			TSharedPtr<FInheritedSubobjectData> TargetData = StaticCastSharedPtr<FInheritedSubobjectData>(TargetParentHandle.GetSharedDataPtr());
+			
+			// Create a new subobject data set with this component. Use the SCS node here and the subobject data
+			// will correctly associate the component template
+			FSubobjectDataHandle NewDataHandle = FactoryCreateSubobjectDataWithParent(
+				NewActorComponent,
+				TargetParentHandle,
+				TargetData ? TargetData->bIsInheritedSCS : false
+			);
+
+			AttachSubobject(TargetParentHandle, NewDataHandle);
+
+			// Map the new subobject's data handle to it's instance name
+			NewNodeMap.Add(NewObjectPair.Key, NewDataHandle);
+		}
+
+		// Restore the node hierarchy from the original copy
+		for (const TPair<FName, FSubobjectDataHandle>& NewNodePair : NewNodeMap)
+		{
+			// If an entry exists in the set of known parent nodes for the current node
+			if (ParentMap.Contains(NewNodePair.Key))
+			{
+				// Get the parent node name
+				FName ParentName = ParentMap[NewNodePair.Key];
+				if (NewNodeMap.Contains(ParentName))
+				{
+					// Reattach the current node to the parent node (this will also handle detachment from the scene root node)
+					const FSubobjectDataHandle& DesiredParentHandle = NewNodeMap[ParentName];
+					AttachSubobject(DesiredParentHandle, NewNodePair.Value);
+				}
+			}
+		}
+
+		// Modify the Blueprint generated class structure (this will also call UpdateTree() as a result)
+		FBlueprintEditorUtils::MarkBlueprintAsStructurallyModified(Blueprint);
+	}
+	else if(AActor* ActorContext = PasteToContextData->GetMutableObject<AActor>())
+	{
+		// Determine where in the hierarchy to paste (default to the root)
+		USceneComponent* TargetComponent = ActorContext->GetRootComponent();
+		for (FSubobjectDataHandle SelectedNode : NewParentHandles)
+		{
+			FSubobjectData* SelectedData = SelectedNode.GetSharedDataPtr().Get();
+			check(SelectedData);
+		
+			if (USceneComponent* SceneComponent = Cast<USceneComponent>(SelectedData->GetMutableComponentTemplate()))
+			{
+				TargetComponent = SceneComponent;
+				break;
+			}
+		}
+		
+		// Paste the components. This uses the current data from the clipboard
+		TArray<UActorComponent*> PastedComponents;
+		FComponentEditorUtils::PasteComponents(PastedComponents, ActorContext, TargetComponent);
+
+		// Create handles for each pasted subobject, attaching it to it's parent as necessary
+		for(UActorComponent* PastedComponent : PastedComponents)
+		{
+			FSubobjectDataHandle ParentHandle = FindHandleForObject(PasteToContext, PastedComponent->GetOuter());
+			FSubobjectDataHandle PastedHandle = FactoryCreateSubobjectDataWithParent(PastedComponent, ParentHandle);
+			if(PastedHandle.IsValid())
+			{
+				OutPastedHandles.AddUnique(PastedHandle);
+			}
+		}
+	}
+}
+
+void USubobjectDataSubsystem::DuplicateSubobjects(const FSubobjectDataHandle& Context, const TArray<FSubobjectDataHandle>& SubobjectsToDup, UBlueprint* BpContext)
+{
+	if(!Context.IsValid() || SubobjectsToDup.IsEmpty())
+	{
+		return;
+	}
+
+	FAddNewSubobjectParams NewSubobjectParams;
+	NewSubobjectParams.BlueprintContext = BpContext;
+	NewSubobjectParams.ParentHandle = Context;
+	NewSubobjectParams.bConformTransformToParent = false;
+	
+	FText FailedAddReason = FText::GetEmpty();
+	
+	TMap<FSubobjectData*, FSubobjectData*> DuplicateSceneComponentMap;
+
+	// For each Subobject to dup, add it as a new subobejct to the context
+	for(const FSubobjectDataHandle& OriginalHandle : SubobjectsToDup)
+	{
+		if(!OriginalHandle.IsValid())
+		{
+			UE_LOG(LogSubobjectSubsystem, Warning, TEXT("Could not duplicate one or more subobjects, an invalid SubobjectToDup was given!"));
+			continue;
+		}
+
+		FSubobjectData* OriginalData = OriginalHandle.GetSharedDataPtr().Get();
+		if(UActorComponent* ComponentTemplate = OriginalData->GetMutableComponentTemplate())
+		{
+			USCS_Node* SCSNode = OriginalData->GetSCSNode();
+			check(SCSNode == nullptr || SCSNode->ComponentTemplate == ComponentTemplate);
+
+			NewSubobjectParams.NewClass = ComponentTemplate->GetClass();
+			NewSubobjectParams.AssetOverride = SCSNode ? (UObject*)SCSNode : ComponentTemplate;
+			
+			FSubobjectDataHandle ClonedSubobject = AddNewSubobject(NewSubobjectParams, FailedAddReason);
+			FSubobjectData* ClonedData = ClonedSubobject.GetSharedDataPtr().Get();
+			
+			if(ClonedData && ClonedData->IsSceneComponent())
+			{
+				DuplicateSceneComponentMap.Add(OriginalData, ClonedData);
+			}
+		}
+	}
+
+	for (const TPair<FSubobjectData*, FSubobjectData*>& DuplicatedPair : DuplicateSceneComponentMap)
+	{
+		FSubobjectData* OriginalData = DuplicatedPair.Key;
+		FSubobjectData* NewData = DuplicatedPair.Value;
+
+		USceneComponent* OriginalComponent = CastChecked<USceneComponent>(OriginalData->GetMutableComponentTemplate());
+		USceneComponent* NewSceneComponent = CastChecked<USceneComponent>(NewData->GetMutableComponentTemplate());
+		
+		if(BpContext)
+		{
+			// Ensure that any native attachment relationship inherited from the original copy is removed (to prevent a GLEO assertion)
+			NewSceneComponent->DetachFromComponent(FDetachmentTransformRules::KeepWorldTransform);
+		}
+
+		// If we're duplicating the root then we're already a child of it so need to reparent, but we do need to reset the scale
+		// otherwise we'll end up with the square of the root's scale instead of being the same size.
+		if(OriginalData->IsDefaultSceneRoot())
+		{
+			NewSceneComponent->SetRelativeScale3D_Direct(FVector(1.f));
+		}
+		else
+		{
+			// If the original node was parented, attempt to add the duplicate as a child of the same parent node if the parent is not
+			// part of the duplicate set, otherwise parent to the parent's duplicate
+			FSubobjectDataHandle ParentHandle = OriginalData->GetParentHandle();
+			if(ParentHandle.IsValid())
+			{
+				if (FSubobjectData** ParentDuplicateComponent = DuplicateSceneComponentMap.Find(ParentHandle.GetData()))
+				{
+					//FSCSEditorTreeNodePtrType DuplicateParentNodePtr = FindTreeNode(*ParentDuplicateComponent);
+					//if (DuplicateParentNodePtr.IsValid())
+					{
+						//ParentNodePtr = DuplicateParentNodePtr;
+					}
+				}
+			}
+		}
+	}
+}
+
+FScopedTransaction* USubobjectDataSubsystem::BeginTransaction(const TArray<FSubobjectDataHandle>& Handles, const FText& Description, UBlueprint* InBlueprint)
+{
+	FScopedTransaction* OutTransaction = new FScopedTransaction(Description);
+	if(InBlueprint)
+	{
+		FBlueprintEditorUtils::MarkBlueprintAsModified(InBlueprint);
+	}
+
+	for (const FSubobjectDataHandle& Handle : Handles)
+	{
+		if(FSubobjectData* Data = Handle.GetData())
+		{
+			if(USCS_Node* SCS_Node = Data->GetSCSNode())
+			{
+				USimpleConstructionScript* SCS = SCS_Node->GetSCS();
+				UBlueprint* Blueprint = SCS ? SCS->GetBlueprint() : nullptr;
+				if (Blueprint == InBlueprint)
+				{
+					SCS_Node->Modify();
+				}
+			}
+
+			// Modify template, any instances will be reconstructed as part of PostUndo:
+			if (UActorComponent* ComponentTemplate = Data->GetMutableObjectForBlueprint<UActorComponent>(InBlueprint))
+			{
+				ComponentTemplate->SetFlags(RF_Transactional);
+				ComponentTemplate->Modify();
+			}
+		}
+	}
+	return OutTransaction;
+}
+
+FSubobjectDataHandle USubobjectDataSubsystem::CreateSubobjectData(UObject* Context, const FSubobjectDataHandle& ParentHandle /* = FSubobjectDataHandle::InvalidHandle */, bool bIsInheritedSCS/* = false */)
+{
+	FCreateSubobjectParams Params;
+	Params.Context = Context;
+	Params.ParentHandle = ParentHandle;
+	Params.bIsInheritedSCS = bIsInheritedSCS;
+
+	ISubobjectDataFactory* Factory = FactoryManager->FindFactoryToUse(Params);
+	check(Factory);
+	TSharedPtr<FSubobjectData> SharedPtr = Factory->CreateSubobjectData(Params);
+	
+	if(!SharedPtr.IsValid())
+	{
+		ensureMsgf(false, TEXT("The subobject data factories failed to create subobject data!"));
+		SharedPtr = TSharedPtr<FSubobjectData>(new FSubobjectData(Context, ParentHandle));
+	}
+
+	check(SharedPtr.IsValid());
+	
+	SharedPtr->Handle.DataPtr = SharedPtr;
+
+	return SharedPtr->GetHandle();
+}
+
+FSubobjectDataHandle USubobjectDataSubsystem::FactoryCreateSubobjectDataWithParent(UObject* Context, const FSubobjectDataHandle& ParentHandle, bool bIsInheritedSCS /* = false */)
+{
+	FSubobjectData* ParentData = ParentHandle.GetSharedDataPtr().Get();
+	if (!ensureMsgf(ParentData, TEXT("Attempted to use an invalid parent subobject handle!")))
+	{
+		return FSubobjectDataHandle::InvalidHandle;
+	}
+
+	// Does this parent subobject already have this object listed in the children?
+	// if yes, then just return that handle
+	const FSubobjectDataHandle& ExistingChild = ParentData->FindChildByObject(Context);
+	if (ExistingChild.IsValid())
+	{
+		return ExistingChild;
+	}
+
+	// Otherwise, we need to create a new handle
+	FSubobjectDataHandle OutHandle = CreateSubobjectData(Context, ParentHandle);
+	
+	// Inform the parent that it has a new child
+	const bool bSuccess = ParentData->AddChildHandleOnly(OutHandle);
+	ensureMsgf(bSuccess, TEXT("Failed to add a child to parent subobject!"));
+
+	return OutHandle;
+}
+
+FSubobjectDataHandle USubobjectDataSubsystem::FactoryCreateInheritedBpSubobject(UObject* Context, const FSubobjectDataHandle& InParentHandle, bool bIsInherited, TArray<FSubobjectDataHandle>& OutArray)
+{
+	USCS_Node* InSCSNode = Cast<USCS_Node>(Context);
+	
+	FSubobjectDataHandle ParentHandle = InParentHandle;
+	FSubobjectData* ParentData = ParentHandle.GetData();
+	
+	check(InSCSNode && ParentData && ParentData->IsValid());
+	// Get a handle for us to work with
+	FSubobjectDataHandle OutHandle = FactoryCreateSubobjectDataWithParent(InSCSNode, ParentHandle);
+	check(OutHandle.IsValid());
+	FSubobjectData* NewData = OutHandle.GetData();
+
+	// @todo the concept of inherited SCS nodes will not exist once the
+	// larger subobject refactor roles out, but this is necessary for now.
+	if(TSharedPtr<FInheritedSubobjectData> InheritedData = StaticCastSharedPtr<FInheritedSubobjectData>(OutHandle.GetSharedDataPtr()))
+	{
+		InheritedData->bIsInheritedSCS = bIsInherited;
+	}
+	
+	// Determine whether or not the given node is inherited from a parent Blueprint
+	USimpleConstructionScript* NodeSCS = InSCSNode->GetSCS();
+	
+	if( InSCSNode->ComponentTemplate && 
+        InSCSNode->ComponentTemplate->IsA(USceneComponent::StaticClass()) && 
+        ParentData->IsComponent())
+	{
+		bool bParentIsEditorOnly = ParentData->GetComponentTemplate()->IsEditorOnly();
+		// if you can't nest this new node under the proposed parent (then swap the two)
+		if (bParentIsEditorOnly && !InSCSNode->ComponentTemplate->IsEditorOnly() && ParentData->CanReparent())
+		{
+			FSubobjectData* OldParentPtr = ParentData;
+			FSubobjectData* GrandparentPtr = OldParentPtr->GetParentHandle().GetData();
+
+			DetachSubobject(OldParentPtr->GetHandle(), NewData->GetHandle());
+			NodeSCS->RemoveNode(OldParentPtr->GetSCSNode());
+			
+			// if the grandparent node is invalid (assuming this means that the parent node was the scene-root)
+			if (!GrandparentPtr->IsValid())
+			{
+				NodeSCS->AddNode(NewData->GetSCSNode());
+			}
+			else 
+			{
+				AttachSubobject(GrandparentPtr->GetHandle(), NewData->GetHandle());
+			}
+			
+			// move the proposed parent in as a child to the new node
+			AttachSubobject(NewData->GetHandle(), OldParentPtr->GetHandle());
+
+		} // if bParentIsEditorOnly...
+	}
+	else
+	{
+		// If the SCS root node array does not already contain the given node, this will add it (this should only occur after node creation)
+		if(NodeSCS != nullptr)
+		{
+			NodeSCS->AddNode(InSCSNode);
+		}
+	}
+
+	// Recursively add the given SCS node's child nodes
+	if(OutHandle.IsValid())
+	{
+		OutArray.Add(OutHandle);
+
+		for (USCS_Node* ChildNode : InSCSNode->GetChildNodes())
+		{
+			FSubobjectDataHandle NewChildHandle = FactoryCreateInheritedBpSubobject(ChildNode, OutHandle, bIsInherited, OutArray);
+			ensure(NewChildHandle.IsValid());
+			OutArray.Add(NewChildHandle);
+		}	
+	}
+	
+	return OutHandle;
+}
+
+void USubobjectDataSubsystem::RenameSubobjectMemberVariable(UBlueprint* BPContext, const FSubobjectDataHandle& InHandle, const FName NewName)
+{
+	if (!BPContext || !InHandle.IsValid())
+	{
+		return;
+	}
+
+	if(FSubobjectData* Data = InHandle.GetSharedDataPtr().Get())
+	{
+		if(USCS_Node* Node = Data->GetSCSNode())
+		{
+			FBlueprintEditorUtils::RenameComponentMemberVariable(BPContext, Node, NewName);
+		}
+	}
+}
+
+#undef LOCTEXT_NAMESPACE
diff --git a/Engine/Source/Editor/SubobjectDataInterface/Public/BaseSubobjectDataFactory.h b/Engine/Source/Editor/SubobjectDataInterface/Public/BaseSubobjectDataFactory.h
new file mode 100644
--- /dev/null
+++ b/Engine/Source/Editor/SubobjectDataInterface/Public/BaseSubobjectDataFactory.h
@@ -0,0 +1,21 @@
+// Copyright Epic Games, Inc. All Rights Reserved.
+
+#pragma once
+
+#include "CoreMinimal.h"
+#include "Templates/SharedPointer.h"
+#include "SubobjectDataFactory.h"
+
+/**
+* This is the default subobject data factory that will provide a simple 
+*/
+class FBaseSubobjectDataFactory : public ISubobjectDataFactory
+{
+public:
+	
+	// Begin ISubobjectDataFactory interface
+	virtual FName GetID() const { return ISubobjectDataFactory::GetDefaultFactoryID(); }
+	virtual TSharedPtr<FSubobjectData> CreateSubobjectData(const FCreateSubobjectParams& Params) override;
+	virtual bool ShouldCreateSubobjectData(const FCreateSubobjectParams& Params) const override;
+	// End ISubobjectDataFactory
+};
\ No newline at end of file
diff --git a/Engine/Source/Editor/SubobjectDataInterface/Public/ChildActorSubobjectData.h b/Engine/Source/Editor/SubobjectDataInterface/Public/ChildActorSubobjectData.h
new file mode 100644
--- /dev/null
+++ b/Engine/Source/Editor/SubobjectDataInterface/Public/ChildActorSubobjectData.h
@@ -0,0 +1,27 @@
+// Copyright Epic Games, Inc. All Rights Reserved.
+
+#pragma once
+
+#include "InheritedSubobjectData.h"
+
+/**
+ * Container that represents the subobejct data of a child actor component.
+ */
+struct SUBOBJECTDATAINTERFACE_API FChildActorSubobjectData final : public FInheritedSubobjectData
+{
+	explicit FChildActorSubobjectData(UObject* ContextObject, const FSubobjectDataHandle& ParentHandle, const bool InbIsInheritedSCS);
+
+	// FSubobjectData interface
+	virtual FText GetDisplayName() const override;
+	virtual FText GetActorDisplayText() const override;
+	virtual bool IsChildActor() const override;
+	virtual bool CanDelete() const override;
+	virtual bool CanReparent() const override;
+	virtual bool CanDuplicate() const override;
+	// End FSubobjectData
+
+	inline const UChildActorComponent* GetChildActorComponent(bool bEvenIfPendingKill = false) const
+	{
+		return GetObject<UChildActorComponent>(bEvenIfPendingKill);
+	}
+};
\ No newline at end of file
diff --git a/Engine/Source/Editor/SubobjectDataInterface/Public/ChildSubobjectDataFactory.h b/Engine/Source/Editor/SubobjectDataInterface/Public/ChildSubobjectDataFactory.h
new file mode 100644
--- /dev/null
+++ b/Engine/Source/Editor/SubobjectDataInterface/Public/ChildSubobjectDataFactory.h
@@ -0,0 +1,21 @@
+// Copyright Epic Games, Inc. All Rights Reserved.
+
+#pragma once
+
+#include "CoreMinimal.h"
+#include "Templates/SharedPointer.h"
+#include "SubobjectDataFactory.h"
+
+/**
+* This is the default subobject data factory that will provide a simple 
+*/
+class FChildSubobjectDataFactory : public ISubobjectDataFactory
+{
+public:
+
+	// Begin ISubobjectDataFactory interface
+	virtual FName GetID() const { return TEXT("ChildSubobjectFactory"); }
+	virtual TSharedPtr<FSubobjectData> CreateSubobjectData(const FCreateSubobjectParams& Params) override;
+	virtual bool ShouldCreateSubobjectData(const FCreateSubobjectParams& Params) const override;
+	// End ISubobjectDataFactory
+};
\ No newline at end of file
diff --git a/Engine/Source/Editor/SubobjectDataInterface/Public/InheritedSubobjectData.h b/Engine/Source/Editor/SubobjectDataInterface/Public/InheritedSubobjectData.h
new file mode 100644
--- /dev/null
+++ b/Engine/Source/Editor/SubobjectDataInterface/Public/InheritedSubobjectData.h
@@ -0,0 +1,26 @@
+// Copyright Epic Games, Inc. All Rights Reserved.
+
+#pragma once
+
+#include "SubobjectData.h"
+
+/**
+ * Container that represents the subobejct data of a child actor component.
+ */
+struct SUBOBJECTDATAINTERFACE_API FInheritedSubobjectData : public FSubobjectData
+{
+	friend class USubobjectDataSubsystem;
+	
+	FInheritedSubobjectData(UObject* ContextObject, const FSubobjectDataHandle& ParentHandle, const bool bIsInheritedSCS);
+	
+	// FSubobjectData interface
+	virtual bool IsNativeComponent() const override;
+	virtual bool CanEdit() const override;
+	virtual bool IsInheritedSCSNode() const override;
+	// End FSubobjectData
+
+protected:
+
+	/** True if this SCS node is inherited from another blueprint generated class. */
+	bool bIsInheritedSCS;
+};
\ No newline at end of file
diff --git a/Engine/Source/Editor/SubobjectDataInterface/Public/InheritedSubobjectDataFactory.h b/Engine/Source/Editor/SubobjectDataInterface/Public/InheritedSubobjectDataFactory.h
new file mode 100644
--- /dev/null
+++ b/Engine/Source/Editor/SubobjectDataInterface/Public/InheritedSubobjectDataFactory.h
@@ -0,0 +1,21 @@
+// Copyright Epic Games, Inc. All Rights Reserved.
+
+#pragma once
+
+#include "CoreMinimal.h"
+#include "Templates/SharedPointer.h"
+#include "SubobjectDataFactory.h"
+
+/**
+* This subobject data factory will create native and blueprint inherited subobjects
+*/
+class FInheritedSubobjectDataFactory : public ISubobjectDataFactory
+{
+public:
+	
+	// Begin ISubobjectDataFactory interface
+	virtual FName GetID() const { return TEXT("InheritedSubobjectFactory"); }
+	virtual TSharedPtr<FSubobjectData> CreateSubobjectData(const FCreateSubobjectParams& Params) override;
+	virtual bool ShouldCreateSubobjectData(const FCreateSubobjectParams& Params) const override;
+	// End ISubobjectDataFactory
+};
\ No newline at end of file
diff --git a/Engine/Source/Editor/SubobjectDataInterface/Public/SubobjectData.h b/Engine/Source/Editor/SubobjectDataInterface/Public/SubobjectData.h
new file mode 100644
--- /dev/null
+++ b/Engine/Source/Editor/SubobjectDataInterface/Public/SubobjectData.h
@@ -0,0 +1,303 @@
+// Copyright Epic Games, Inc. All Rights Reserved.
+
+#pragma once
+
+#include "CoreMinimal.h"
+#include "Templates/SharedPointer.h"
+#include "SubobjectDataHandle.h"
+#include "Components/ActorComponent.h"
+#include "Kismet/BlueprintFunctionLibrary.h"
+#include "Engine/SCS_Node.h"		// #TODO_BH  We need to remove this when the actual subobject refactor happens
+
+#include "SubobjectData.generated.h"
+
+class UBlueprint;
+
+/**
+* A struct that represents a single subobject. This can be anything, but are
+* most commonly components attached to an actor instance or blueprint. Keeps track
+* of the handles to its parent object and any child that it has.
+*
+* If you wish to modify a subobject, use the SubobjectDataSubsystem.
+*/
+USTRUCT(BlueprintType)
+struct SUBOBJECTDATAINTERFACE_API FSubobjectData
+{
+friend class USubobjectDataSubsystem;
+public:
+	GENERATED_BODY()
+
+	FSubobjectData();
+
+	FSubobjectData(UObject* ContextObject, const FSubobjectDataHandle& ParentHandle);
+
+	virtual ~FSubobjectData();
+
+	/** True if this subobject's handle is valid */
+	bool IsValid() const { return Handle.IsValid(); }
+	
+	/**
+	 * @return Whether or not we can edit properties for this subobject
+	 */
+	virtual bool CanEdit() const;
+
+	/**
+	 * @return Whether or not this object represents a subobject that can be deleted
+	 */
+	virtual bool CanDelete() const;
+
+	/**
+	* @return Whether or not this object represents a subobject that can be duplicated
+	*/
+	virtual bool CanDuplicate() const;
+
+	/**
+	* @return Whether or not this object represents a subobject that can be copied
+	*/
+	virtual bool CanCopy() const;
+
+	/**
+	 * @return Whether or not this object represents a subobject that can 
+	 * be reparented to other subobjects based on its context.
+	 */
+	virtual bool CanReparent() const;
+
+	/**
+	* @return Whether or not this object represents a subobject that can 
+	* be renamed.
+	*/
+	virtual bool CanRename() const;
+
+	/**
+	* @return	A read-only reference to the object represented by this node.
+	*/
+	inline const UObject* GetObject(bool bEvenIfPendingKill = false) const { return WeakObjectPtr.Get(bEvenIfPendingKill); }
+
+	/**
+     * @param	InBlueprint			The Blueprint in which the object will be edited.
+     *
+     * @note	May not be the same as the value returned by GetObject().
+     * @return	A reference to the object represented by this node that can be modified within the given Blueprint.
+     */
+	const UObject* GetObjectForBlueprint(UBlueprint* Blueprint) const;
+
+	template<class T>
+	inline const T* GetObjectForBlueprint(UBlueprint* Blueprint) const
+	{
+		return Cast<T>(GetObjectForBlueprint(Blueprint));
+	}
+	
+	/**
+	 * @param	bEvenIfPendingKill	If false, nullptr will be returned if the cached object instance is pending kill.
+	 *								If true, it will be returned regardless (this is used for recaching the object if the objects
+	 *								have been reinstanced following construction script execution).
+	 *
+	 * @note	Deliberately non-virtual, for performance reasons.
+	 * @return	A read-only reference to the object represented by this node.
+	 */
+	template<class T>
+	inline const T* GetObject(bool bEvenIfPendingKill = false) const
+	{
+		return Cast<T>(WeakObjectPtr.Get(bEvenIfPendingKill));
+	}
+
+	/**
+	 * @param	bEvenIfPendingKill	If false, nullptr will be returned if the cached component template is pending kill.
+	 *								If true, it will be returned regardless (this is used for recaching the component template if the objects
+	 *								have been reinstanced following construction script execution).
+	 *
+	 * @note	Deliberately non-virtual, for performance reasons.
+	 * @warning This will not return the right component for components overridden by the inherited component handler, you need to call GetOrCreateEditableComponentTemplate instead
+	 * @return	The component template or instance represented by this node, if it's a component node.
+	 */
+	inline const UActorComponent* GetComponentTemplate(bool bEvenIfPendingKill = false) const
+	{
+		return Cast<UActorComponent>(WeakObjectPtr.Get(bEvenIfPendingKill));
+	}
+
+	inline const UActorComponent* FindComponentInstanceInActor(const AActor* InActor) const { return FindMutableComponentInstanceInActor(InActor); }
+
+	UBlueprint* GetBlueprint() const;
+
+	virtual bool IsInstancedComponent() const;
+	virtual bool IsInstancedActor() const;
+	virtual bool IsNativeComponent() const;
+	virtual bool IsInheritedComponent() const;
+	virtual bool IsSceneComponent() const;
+	virtual bool IsRootComponent() const;
+	virtual bool IsDefaultSceneRoot() const;
+	/* Returns true if this subobject is a component. */
+	virtual bool IsComponent() const;
+	virtual bool IsChildActor() const;
+	virtual bool IsRootActor() const;
+	virtual bool IsActor() const;
+	/** True if this subobject is an instance of an inherited component (its owner is an instanced actor) */
+	virtual bool IsInstancedInheritedComponent() const;
+
+	bool IsAttachedTo(const FSubobjectDataHandle& InHandle) const;
+
+	virtual FString GetDisplayString(bool bShowNativeComponentNames = true) const;
+
+	/**
+	 * Returns a string that gives details about this subobject. For Inherited
+	 * components it will return "(Inherited)". For native components
+	 * this function will return "(NativeComponentName) (Inherited)"
+	 */
+	virtual FText GetDisplayNameContextModifiers(bool bShowNativeComponentNames = true) const;
+	
+	virtual FText GetDisplayName() const;
+
+	virtual FName GetVariableName() const;
+
+	// Sockets for attaching in the viewport
+	FText GetSocketName() const;
+	FName GetSocketFName() const;
+	bool HasValidSocket() const;
+	void SetSocketName(FName InNewName);
+	void SetupAttachment(FName SocketName, const FSubobjectDataHandle& AttachParentHandle = FSubobjectDataHandle::InvalidHandle);
+
+	FSubobjectDataHandle GetHandle() const { return Handle; }
+	FSubobjectDataHandle GetParentHandle() const { return ParentObjectHandle; }
+	FSubobjectData* GetParentData() const { return ParentObjectHandle.GetData(); }
+	bool HasParent() const { return ParentObjectHandle.IsValid(); }
+	const TArray<FSubobjectDataHandle>& GetChildrenHandles() const { return ChildrenHandles; }
+	
+	/** Walks the parent heirarchy and returns the root subobject in this case */
+	FSubobjectDataHandle GetRootSubobject() const;
+
+	// Return true if the given handle is a child of this subobject
+	bool HasChild(const FSubobjectDataHandle& ChildHandle) const;
+
+	// Attempt to find the given handle in out child array and return a pointer to it
+	FSubobjectDataHandle FindChild(const FSubobjectDataHandle& ChildHandle) const;
+	FSubobjectDataHandle FindChildByObject(UObject* ContextObject) const;
+
+	/** Get the asset name of this subobject from the asset brokerage */
+	FText GetAssetName() const;
+
+	/** Get the asset path of this subobject from the asset brokerage */
+	FText GetAssetPath() const;
+
+	/** Returns true if the asset this object represents is visible */
+	bool IsAssetVisible() const;
+
+	/**
+	* Retrieves tooltip text describing the specified component's mobility.
+	*
+	* @returns An FText object containing a description of the component's mobility
+	*/
+	FText GetMobilityToolTipText() const;
+
+	/**
+	* Retrieves a tooltip text describing if the component is marked Editor only or not
+	*
+	* @returns An FText object containing a description of if the component is marked Editor only or not
+	*/
+	FText GetComponentEditorOnlyTooltipText() const;
+
+	/**
+	* Retrieves tooltip text describing where the component was first introduced (for inherited components).
+	* 
+	* @returns An FText object containing a description of when the component was first introduced
+	*/
+	FText GetIntroducedInToolTipText() const;
+
+	virtual FText GetActorDisplayText() const;
+	
+	FLinearColor GetColorTintForIcon() const;
+	
+protected:
+	
+	/**
+	* Add the given subobject handle as a child of this.
+	* NOTE: This does NOT do any actual manipulation of the subobject structure.
+	* This only updates this subobject to have the correct handles in it's child structure.
+	* Actual manipulation of subobjects should be done via the SubobjectSubsystem.
+	*/
+	bool AddChildHandleOnly(const FSubobjectDataHandle& InHandle);
+
+	/**
+	* Add the given subobject handle as a child of this.
+	* 	* NOTE: This does NOT do any actual manipulation of the subobject structure.
+	* This only updates this subobject to have the correct handles in it's child structure.
+	* Actual manipulation of subobjects should be done via the SubobjectSubsystem.
+	*/
+	bool RemoveChildHandleOnly(const FSubobjectDataHandle& InHandle);
+
+	void SetParentHandle(const FSubobjectDataHandle& NewParentHandle) { ParentObjectHandle = NewParentHandle; }
+	void ClearParentHandle() { ParentObjectHandle = FSubobjectDataHandle::InvalidHandle; }
+	
+	// Weak ptr to the object instance represented by this node (e.g. component template)
+	UPROPERTY()
+	TWeakObjectPtr<UObject> WeakObjectPtr;
+
+	// Handle that will have a globally unique ID when it is constructed (i.e. when this is constructed)
+	UPROPERTY()
+	FSubobjectDataHandle Handle;
+
+	// Handle to this object's parent. If this is invalid, then it is the root subobject.
+	UPROPERTY()
+	FSubobjectDataHandle ParentObjectHandle;
+
+	// Set of child subobjects that use this subobject as a parent
+	UPROPERTY()
+	TArray<FSubobjectDataHandle> ChildrenHandles;
+
+	////////////////////////////////////////////
+	// Muteable accessors for this subobject to be used by the subsystem
+	inline UObject* GetMutableObject(bool bEvenIfPendingKill = false) const { return WeakObjectPtr.Get(bEvenIfPendingKill); }
+
+	// If this subobject is an actor, then return a pointer to that. Otherwise, search for
+	// the owning actor on this component and return it.
+	AActor* GetMutableActorContext();
+
+	template<class T>
+    inline T* GetMutableObject(bool bEvenIfPendingKill = false) const
+	{
+		return Cast<T>(WeakObjectPtr.Get(bEvenIfPendingKill));
+	}
+
+	template<class T>
+	inline T* GetMutableObjectForBlueprint(UBlueprint* Blueprint) const
+    {
+    	return const_cast<T*>(GetObjectForBlueprint<T>(Blueprint));
+    }
+	
+	inline UActorComponent* GetMutableComponentTemplate(bool bEvenIfPendingKill = false) const
+	{
+		return Cast<UActorComponent>(WeakObjectPtr.Get(bEvenIfPendingKill));
+	}
+
+	UActorComponent* FindMutableComponentInstanceInActor(const AActor* InActor) const;
+
+	////////////////////////////////////////////////////
+	// Anything related to an SCS node will be changed with an upcoming refactor and should remain private
+protected:
+
+	// Tries to find a SCS node that was likely responsible for creating the specified instance component.  Note: This is not always possible to do!
+	static USCS_Node* FindSCSNodeForInstance(const UActorComponent* InstanceComponent, UClass* ClassToSearch);
+
+	virtual USCS_Node* GetSCSNode(bool bEvenIfPendingKill = false) const;
+
+	/**
+	* Attempts to set the SCS node member variable based on the WeakObjectPtr.
+	* If the weak obj ptr _is_ an USCS_Node type, then we have to change the component
+	* template.
+	*
+	* If this is an instanced component, then we can find the scs node via FindSCSNodeForInstance.
+	*
+	* @return True if an SCS node has been set
+	*/
+	bool AttemptToSetSCSNode();
+
+	// For now, we need a pointer to the SCS node in order to have the correct component template
+	// and to be able to get the UBlueprint* from components. This will only be set upon construction
+	// if the given Object Context is an USCS_Node. The context object will be set to its template. 
+	TWeakObjectPtr<USCS_Node> SCSNodePtr;
+
+	/**
+	 * True if this SCS node's blueprint is a child of another Blueprint-generated class,
+	 * which means that we must use the InheritableComponentHandler
+	 */
+	virtual bool IsInheritedSCSNode() const;
+};
\ No newline at end of file
diff --git a/Engine/Source/Editor/SubobjectDataInterface/Public/SubobjectDataBlueprintFunctionLibrary.h b/Engine/Source/Editor/SubobjectDataInterface/Public/SubobjectDataBlueprintFunctionLibrary.h
new file mode 100644
--- /dev/null
+++ b/Engine/Source/Editor/SubobjectDataInterface/Public/SubobjectDataBlueprintFunctionLibrary.h
@@ -0,0 +1,119 @@
+// Copyright Epic Games, Inc. All Rights Reserved.
+
+#pragma once
+
+#include "CoreMinimal.h"
+#include "Kismet/BlueprintFunctionLibrary.h"
+#include "SubobjectData.h"
+
+#include "SubobjectDataBlueprintFunctionLibrary.generated.h"
+
+class UBlueprint;
+
+/**
+ * A function library with wrappers around the getter/setter functions for FSubobjectData
+ * that will make it easier to use within blueprint contexts.
+ */
+UCLASS()
+class SUBOBJECTDATAINTERFACE_API USubobjectDataBlueprintFunctionLibrary : public UBlueprintFunctionLibrary
+{
+	GENERATED_BODY()
+public:
+	UFUNCTION(BlueprintCallable, BlueprintPure, Category="Subobject Data")
+	static bool IsValid(const FSubobjectData& Data) { return Data.IsValid(); }
+
+	UFUNCTION(BlueprintCallable, BlueprintPure, Category="Subobject Data")
+	static bool IsHandleValid(const FSubobjectDataHandle& DataHandle) { return DataHandle.IsValid(); }
+
+	UFUNCTION(BlueprintCallable, BlueprintPure, Category="Subobject Data")
+    static void GetData(const FSubobjectDataHandle& DataHandle, FSubobjectData& OutData);
+	
+	/**
+	* @return Get the handle for this subobject data
+	*/
+	UFUNCTION(BlueprintCallable, BlueprintPure, Category = "Subobject Data")
+	static void GetHandle(const FSubobjectData& Data, FSubobjectDataHandle& OutHandle) { OutHandle = Data.GetHandle(); }
+
+	UFUNCTION(BlueprintCallable, BlueprintPure, Category = "Subobject Data")
+	FText GetDisplayName(const FSubobjectData& Data) { return Data.GetDisplayName(); }
+
+	UFUNCTION(BlueprintCallable, BlueprintPure, Category = "Subobject Data")
+	static FName GetVariableName(const FSubobjectData& Data) { return Data.GetVariableName(); }
+
+	UFUNCTION(BlueprintCallable, BlueprintPure, Category = "Subobject Data")
+	static bool IsAttachedTo(const FSubobjectData& Data, const FSubobjectDataHandle& InHandle) { return Data.IsAttachedTo(InHandle); }
+
+	/**
+	* @return Whether or not we can edit properties for this subobject
+	*/
+	UFUNCTION(BlueprintCallable, BlueprintPure, Category="Subobject Data")
+	static bool CanEdit(const FSubobjectData& Data) { return Data.CanEdit(); }
+
+	/**
+	* @return Whether or not this object represents a subobject that can be deleted
+	*/
+	UFUNCTION(BlueprintCallable, BlueprintPure, Category="Subobject Data")
+	static bool CanDelete(const FSubobjectData& Data) { return Data.CanDelete(); }
+
+	/**
+	* @return Whether or not this object represents a subobject that can be duplicated
+	*/
+	UFUNCTION(BlueprintCallable, BlueprintPure, Category="Subobject Data")
+	static bool CanDuplicate(const FSubobjectData& Data) { return Data.CanDuplicate(); }
+
+	/**
+	* @return Whether or not this object represents a subobject that can be copied
+	*/
+	UFUNCTION(BlueprintCallable, BlueprintPure, Category="Subobject Data")
+	static bool CanCopy(const FSubobjectData& Data) { return Data.CanCopy(); }
+
+	/**
+	* @return Whether or not this object represents a subobject that can 
+	* be reparented to other subobjects based on its context.
+	*/
+	UFUNCTION(BlueprintCallable, BlueprintPure, Category="Subobject Data")
+	static bool CanReparent(const FSubobjectData& Data) { return Data.CanReparent(); }
+
+	UFUNCTION(BlueprintCallable, BlueprintPure, Category="Subobject Data")
+	static bool CanRename(const FSubobjectData& Data) { return Data.CanRename(); }
+	
+	UFUNCTION(BlueprintCallable, BlueprintPure, Category="Subobject Data")
+	static const UObject* GetObject(const FSubobjectData& Data, bool bEvenIfPendingKill = false) { return Data.GetObject(bEvenIfPendingKill); }
+
+	UFUNCTION(BlueprintCallable, BlueprintPure, Category="Subobject Data")
+	UBlueprint* GetBlueprint(const FSubobjectData& Data) { return Data.GetBlueprint(); }
+
+	UFUNCTION(BlueprintCallable, BlueprintPure, Category="Subobject Data")
+	static bool IsInstancedComponent(const FSubobjectData& Data) { return Data.IsInstancedComponent(); }
+
+	UFUNCTION(BlueprintCallable, BlueprintPure, Category="Subobject Data")
+	static bool IsInstancedActor(const FSubobjectData& Data) { return Data.IsInstancedActor(); }
+
+	UFUNCTION(BlueprintCallable, BlueprintPure, Category="Subobject Data")
+	static bool IsNativeComponent(const FSubobjectData& Data) { return Data.IsNativeComponent(); }
+
+	UFUNCTION(BlueprintCallable, BlueprintPure, Category="Subobject Data")
+	static bool IsInheritedComponent(const FSubobjectData& Data) { return Data.IsInheritedComponent(); }
+
+	UFUNCTION(BlueprintCallable, BlueprintPure, Category="Subobject Data")
+	static bool IsSceneComponent(const FSubobjectData& Data) { return Data.IsSceneComponent(); }
+
+	UFUNCTION(BlueprintCallable, BlueprintPure, Category="Subobject Data")
+	static bool IsRootComponent(const FSubobjectData& Data) { return Data.IsRootComponent(); }
+
+	UFUNCTION(BlueprintCallable, BlueprintPure, Category="Subobject Data")
+	static bool IsDefaultSceneRoot(const FSubobjectData& Data) { return Data.IsDefaultSceneRoot(); }
+
+	/* Returns true if this subobject is a component. */
+	UFUNCTION(BlueprintCallable, BlueprintPure, Category="Subobject Data")
+	static bool IsComponent(const FSubobjectData& Data) { return Data.IsComponent(); }
+
+	UFUNCTION(BlueprintCallable, BlueprintPure, Category="Subobject Data")
+	static bool IsChildActor(const FSubobjectData& Data) { return Data.IsChildActor(); }
+
+	UFUNCTION(BlueprintCallable, BlueprintPure, Category="Subobject Data")
+	static bool IsRootActor(const FSubobjectData& Data) { return Data.IsRootActor(); }
+	
+	UFUNCTION(BlueprintCallable, BlueprintPure, Category="Subobject Data")
+	static bool IsActor(const FSubobjectData& Data) { return Data.IsActor(); }
+};
\ No newline at end of file
diff --git a/Engine/Source/Editor/SubobjectDataInterface/Public/SubobjectDataFactory.h b/Engine/Source/Editor/SubobjectDataInterface/Public/SubobjectDataFactory.h
new file mode 100644
--- /dev/null
+++ b/Engine/Source/Editor/SubobjectDataInterface/Public/SubobjectDataFactory.h
@@ -0,0 +1,111 @@
+// Copyright Epic Games, Inc. All Rights Reserved.
+
+#pragma once
+
+#include "CoreMinimal.h"
+#include "Templates/SharedPointer.h"
+#include "SubobjectData.h"
+#include "SubobjectDataHandle.h"
+
+/** Options for creating a new subobject data via the factory */
+struct SUBOBJECTDATAINTERFACE_API FCreateSubobjectParams
+{
+	UObject* Context = nullptr;
+	
+	FSubobjectDataHandle ParentHandle = FSubobjectDataHandle::InvalidHandle;
+	
+	bool bIsInheritedSCS = false;
+};
+
+/**
+* A subobject data factory will handle the creation of subobject data. 
+* To create a subobject data factory, override this class and register it 
+* with the subobject data subsystem. 
+*/
+class SUBOBJECTDATAINTERFACE_API ISubobjectDataFactory
+{
+public:
+
+	static FName GetDefaultFactoryID()
+	{
+		static FName ID(TEXT("DefaultSubobjectDataFactory"));
+		return ID;
+	}
+
+	virtual FName GetID() const = 0;
+	
+	virtual TSharedPtr<FSubobjectData> CreateSubobjectData(const FCreateSubobjectParams& Params) = 0;
+
+	virtual bool ShouldCreateSubobjectData(const FCreateSubobjectParams& Params) const = 0;
+};
+
+/**
+ * The factory manager keeps a map of all registered factories that can be used
+ * to spawn Subobject Data. To add your own factory, create a subclass of ISubobjectDataFactory.
+ *
+ * See BaseSubobjectDataFactory.h for an example.
+ */
+class SUBOBJECTDATAINTERFACE_API FSubobjectFactoryManager
+{
+public:
+	explicit FSubobjectFactoryManager() = default;
+	~FSubobjectFactoryManager() = default;
+
+	/***
+	 * Map the given factory to this manager with its ID if it does not already
+	 * exist in the list
+	 * 
+	 * @return True if the factory was registered successfully
+	 */
+	bool RegisterFactory(const TSharedPtr<ISubobjectDataFactory>& Factory)
+	{
+		if(Factory.IsValid())
+		{
+			const FName ID = Factory->GetID();
+			if (!IsFactoryRegistered(ID))
+			{
+				RegisteredFactories.Add(ID, Factory);
+				return true;
+			}
+		}
+		return false;
+	}
+
+	/**
+	 * Remove the factory with the given ID from the registered map.
+	 * @return True if it was removed. 
+	 */
+	bool UnregisterFactory(FName ID)
+	{
+		return RegisteredFactories.Remove(ID) > 0;
+	}
+
+	/** returns true if the given factory ID is registered */
+	bool IsFactoryRegistered(FName ID) const 
+	{
+		return RegisteredFactories.Contains(ID);
+	}
+
+	/**
+	 * Find a factory that returns true from ShouldCreateSubobjectData based on
+	 * the given parameters. If none are found, then the default will be returned. 
+	 */
+	ISubobjectDataFactory* FindFactoryToUse(const FCreateSubobjectParams& Params)
+	{
+		for(const TPair<FName, TSharedPtr<ISubobjectDataFactory>>& Pair : RegisteredFactories)
+		{
+			const TSharedPtr<ISubobjectDataFactory>& Factory = Pair.Value;
+			if(Factory->ShouldCreateSubobjectData(Params))
+			{
+				return Factory.Get();
+			}
+		}
+		
+		TSharedPtr<ISubobjectDataFactory>* Default = RegisteredFactories.Find(ISubobjectDataFactory::GetDefaultFactoryID());
+		return Default ? (*Default).Get() : nullptr;
+	}
+	
+private:
+	/** Map of all registered factories */
+	TMap<FName, TSharedPtr<ISubobjectDataFactory>> RegisteredFactories;
+};
\ No newline at end of file
diff --git a/Engine/Source/Editor/SubobjectDataInterface/Public/SubobjectDataHandle.h b/Engine/Source/Editor/SubobjectDataInterface/Public/SubobjectDataHandle.h
new file mode 100644
--- /dev/null
+++ b/Engine/Source/Editor/SubobjectDataInterface/Public/SubobjectDataHandle.h
@@ -0,0 +1,64 @@
+// Copyright Epic Games, Inc. All Rights Reserved.
+
+#pragma once
+
+#include "CoreMinimal.h"
+#include "Templates/SharedPointer.h"
+#include "SubobjectDataHandle.generated.h"
+
+struct FSubobjectData;
+
+/**
+* A subobject handle is a globally unique identifier for subobjects
+* Upon construction, the handle will be invalid. It is the responsibility
+* of the owning FSubobjectData to set the DataPtr once the subobject
+* data has validated that it has a good context.
+*/
+USTRUCT(BlueprintType)
+struct SUBOBJECTDATAINTERFACE_API FSubobjectDataHandle
+{
+	friend struct FSubobjectData;
+	friend class USubobjectDataSubsystem;
+	
+	GENERATED_USTRUCT_BODY();
+
+	explicit FSubobjectDataHandle()
+		: DataPtr(nullptr)
+	{}
+	
+	~FSubobjectDataHandle()
+	{
+		DataPtr = nullptr;
+	}
+
+	/**
+	* True if the Handle is valid (i.e. not INDEX_NONE). This is true once GenerateNewHandle is called
+	*
+	* @return bool		True if this handle is valid
+	*/
+	inline bool IsValid() const { return DataPtr != nullptr; }
+
+	bool operator==(const FSubobjectDataHandle& Other) const
+	{
+		return DataPtr.Get() == Other.DataPtr.Get();
+	}
+
+	bool operator!=(const FSubobjectDataHandle& Other) const
+	{
+		return DataPtr != Other.DataPtr;
+	}
+
+	/** Returns a pointer to the subobject data that this is a handle for */
+	inline TSharedPtr<FSubobjectData> GetSharedDataPtr() const { return DataPtr; }
+	inline FSubobjectData* GetData() const { return DataPtr.IsValid() ? DataPtr.Get() : nullptr; }
+
+	/**
+	* A static representation of an invalid handle. Mostly used as a return
+	* value for functions that need to return a handle but have the risk of being invalid.
+	*/
+	static const FSubobjectDataHandle InvalidHandle;
+	
+private:
+	/** Pointer to the actual subobject data that this handle represents */
+	TSharedPtr<FSubobjectData> DataPtr;
+};
\ No newline at end of file
diff --git a/Engine/Source/Editor/SubobjectDataInterface/Public/SubobjectDataInterfaceModule.h b/Engine/Source/Editor/SubobjectDataInterface/Public/SubobjectDataInterfaceModule.h
new file mode 100644
--- /dev/null
+++ b/Engine/Source/Editor/SubobjectDataInterface/Public/SubobjectDataInterfaceModule.h
@@ -0,0 +1,14 @@
+// Copyright Epic Games, Inc. All Rights Reserved.
+
+#pragma once
+
+#include "Modules/ModuleManager.h"
+
+class FSubobjectDataInterfaceModule : public IModuleInterface
+{
+public:
+
+	/** IModuleInterface implementation */
+	virtual void StartupModule() override;
+	virtual void ShutdownModule() override;
+};
\ No newline at end of file
diff --git a/Engine/Source/Editor/SubobjectDataInterface/Public/SubobjectDataSubsystem.h b/Engine/Source/Editor/SubobjectDataInterface/Public/SubobjectDataSubsystem.h
new file mode 100644
--- /dev/null
+++ b/Engine/Source/Editor/SubobjectDataInterface/Public/SubobjectDataSubsystem.h
@@ -0,0 +1,355 @@
+// Copyright Epic Games, Inc. All Rights Reserved.
+
+#pragma once
+
+#include "CoreMinimal.h"
+#include "Subsystems/EngineSubsystem.h"
+#include "SubobjectData.h"
+#include "ModuleDescriptor.h"
+#include "SubobjectDataFactory.h"
+#include "SubobjectDataSubsystem.generated.h"
+
+class UActorComponent;
+class FScopedTransaction;
+
+/** Options when adding a new subobject */
+USTRUCT(BlueprintType)
+struct FAddNewSubobjectParams
+{
+	GENERATED_USTRUCT_BODY()
+
+	FAddNewSubobjectParams()
+		: ParentHandle(FSubobjectDataHandle::InvalidHandle)
+		, NewClass(nullptr)
+		, AssetOverride(nullptr)
+		, BlueprintContext(nullptr)
+		, bSkipMarkBlueprintModified(false)
+		, bConformTransformToParent(true)
+	{
+	}
+	
+	UPROPERTY(BlueprintReadWrite, Category = "SubobjectDataSubsystem|AddNewSubobject")
+	FSubobjectDataHandle ParentHandle;
+
+	/** The class of the new subobject that will be added */
+	UPROPERTY(BlueprintReadWrite, Category = "SubobjectDataSubsystem|AddNewSubobject")
+	UClass* NewClass;
+
+	/** Specific asset to use instead of the selected asset in the content browser */
+	UObject* AssetOverride;
+
+	/**
+	 * Pointer to the blueprint context that this subobject is in. If this is null, it is assumed that
+	 * this subobject is being added to an instance.
+	 */
+	UPROPERTY(BlueprintReadWrite, Category = "SubobjectDataSubsystem|AddNewSubobject")
+	UBlueprint* BlueprintContext;
+	
+	/** Optionally skip marking this blueprint as modified (e.g. if we're handling that externally */
+	UPROPERTY(BlueprintReadWrite, Category = "SubobjectDataSubsystem|AddNewSubobject")
+	uint8 bSkipMarkBlueprintModified : 1;
+
+	/** Whether the newly created component should keep its transform, or conform it to its parent */
+	UPROPERTY(BlueprintReadWrite, Category = "SubobjectDataSubsystem|AddNewSubobject")
+	uint8 bConformTransformToParent : 1;
+};
+
+/** Options for reparenting subobjects */
+USTRUCT(BlueprintType)
+struct FReparentSubobjectParams
+{
+	GENERATED_USTRUCT_BODY()
+
+	FReparentSubobjectParams() = default;
+
+	/** The handle of the subobject to reparent to. */
+	UPROPERTY(BlueprintReadWrite, Category = "SubobjectDataSubsystem|AddNewSubobject")
+	FSubobjectDataHandle NewParentHandle;
+
+	/**
+	* Pointer to the blueprint context that this subobject is in. If this is null, it is assumed that
+	* this subobject is being added to an instance.
+	*/
+	UPROPERTY(BlueprintReadWrite, Category = "SubobjectDataSubsystem|AddNewSubobject")
+	UBlueprint* BlueprintContext = nullptr;
+
+	/** 
+	* The preview actor context to be used if in a blueprint context. 
+	* This must have a value if BlueprintContext is needed. 
+	*/
+	UPROPERTY(BlueprintReadWrite, Category = "SubobjectDataSubsystem|AddNewSubobject")
+	AActor* ActorPreviewContext = nullptr;
+};
+
+/**
+* The Subobject Data Subsystem will produce the reflected subobject data
+* based on a given root object. A root object can be anything, an actor 
+* instance clicked on via the level editor, a UBlueprint* by opening an asset,
+* or something piped in from python or other scripting languages. 
+*/
+UCLASS()
+class SUBOBJECTDATAINTERFACE_API USubobjectDataSubsystem : public UEngineSubsystem
+{
+	GENERATED_BODY()
+public:
+
+	/** Implement this for initialization of instances of the system */
+	virtual void Initialize(FSubsystemCollectionBase& Collection) override;
+
+	/** Implement this for deinitialization of instances of the system */
+	virtual void Deinitialize() override;
+	
+	/**
+	* Static wrapper for getting this engine subsystem. Will return nullptr
+	* if the SubobjectDataInterface module has not been loaded yet.
+	*/
+	static USubobjectDataSubsystem* Get();
+	
+	/**
+	* Gather all subobjects that the given UObject context has. Populates an array of 
+	* handles that will have the given context and all it's subobjects.
+	* 
+	* @param Context		Object to gather subobjects for
+	* @param OutArray		Array to populate (will be emptied first)
+	*/
+	void GatherSubobjectData(UObject* Context, TArray<FSubobjectDataHandle>& OutArray);
+
+	/**
+	* Gather all subobjects that the given Blueprint context has. Populates an array of
+	* handles that will have the given context and all it's subobjects.
+	*
+	* @param Context		Object to gather subobjects for
+	* @param OutArray		Array to populate (will be emptied first)
+	*/
+	UFUNCTION(BlueprintCallable, Category = "SubobjectDataSubsystem", meta = (DisplayName = "Gather Subobject Data For Blueprint"))
+	void K2_GatherSubobjectDataForBlueprint(UBlueprint* Context, TArray<FSubobjectDataHandle>& OutArray);
+
+	/**
+	* Gather all subobjects that the given actor instance has. Populates an array of
+	* handles that will have the given context and all it's subobjects.
+	*
+	* @param Context		Object to gather subobjects for
+	* @param OutArray		Array to populate (will be emptied first)
+	*/
+	UFUNCTION(BlueprintCallable, Category = "SubobjectDataSubsystem", meta = (DisplayName = "Gather Subobject Data For Instance"))
+	void K2_GatherSubobjectDataForInstance(AActor* Context, TArray<FSubobjectDataHandle>& OutArray);
+
+	/**
+	* Attempt to find the subobject data for a given handle. OutData will only 
+	* be valid if the function returns true.
+	*
+	* @param Handle		Handle of the subobject data you want to aquire
+	* @param OutData	Reference to the subobject data to populate
+	*
+	* @return bool		true if the data was found
+	*/
+	UFUNCTION(BlueprintCallable, Category = "SubobjectDataSubsystem", meta = (DisplayName = "FindSubobjectDataFromHandle"))
+	bool K2_FindSubobjectDataFromHandle(const FSubobjectDataHandle& Handle, FSubobjectData& OutData) const;
+
+	/**
+	* Attempt to find an existing handle for the given object. 
+	*
+	* @param Context		The context that the object to find is within
+	* @param ObjectToFind	The object that you want to find the handle for within the context
+	*
+	* @return FSubobjectDataHandle	The subobject handle for the object, Invalid handle if not found.
+	*/
+	UFUNCTION(BlueprintCallable, Category = "SubobjectDataSubsystem")
+	FSubobjectDataHandle FindHandleForObject(const FSubobjectDataHandle& Context, const UObject* ObjectToFind, UBlueprint* BPContext = nullptr) const;
+	
+	////////////////////////////////////////////
+	// Modifying subobjects
+
+	/**
+	 * Creates a new C++ component from the specified class type
+	 * The user will be prompted to pick a new subclass name and code will be recompiled
+	 *
+	 * @return The new class that was created
+	 */
+	// #TODO_BH Finish up this function before exposing it to BP, it has some kinks to
+	// work out related to loading the newly created class.
+	//UFUNCTION(BlueprintCallable, Category = "SubobjectDataSubsystem")
+	static UClass* CreateNewCPPComponent(TSubclassOf<UActorComponent> ComponentClass, const FString& NewClassPath, const FString& NewClassName);
+
+	/**
+	 * Creates a new Blueprint component from the specified class type
+	 * The user will be prompted to pick a new subclass name and a blueprint asset will be created
+	 *
+	 * @return The new class that was created
+	 */
+	UFUNCTION(BlueprintCallable, Category = "SubobjectDataSubsystem")
+	static UClass* CreateNewBPComponent(TSubclassOf<UActorComponent> ComponentClassconst, const FString& NewClassPath, const FString& NewClassName);
+	
+	/**
+	* Add a new subobject as a child to the given parent object 
+	*
+	* @param Params			Options to consider when adding this subobject
+	*
+	* @return FSubobjectDataHandle		Handle to the newly created subobject, Invalid handle if creation failed
+	*/
+	UFUNCTION(BlueprintCallable, Category = "SubobjectDataSubsystem")
+	FSubobjectDataHandle AddNewSubobject(const FAddNewSubobjectParams& Params, FText& FailReason);
+
+	/**
+	* Attempts to delete the given array of subobjects from their context 
+	*
+	* @param ContextHandle			The owning context of the subobjects that should be removed
+	* @param SubobjectsToDelete		Array of subobject handles that should be deleted
+	*
+	* @return 	The number of subobjects successfully deleted
+	*/
+	UFUNCTION(BlueprintCallable, Category = "SubobjectDataSubsystem")
+	int32 DeleteSubobjects(const FSubobjectDataHandle& ContextHandle, const TArray<FSubobjectDataHandle>& SubobjectsToDelete, UBlueprint* BPContext = nullptr);
+
+	UFUNCTION(BlueprintCallable, Category = "SubobjectDataSubsystem")
+    int32 DeleteSubobject(const FSubobjectDataHandle& ContextHandle, const FSubobjectDataHandle& SubobjectToDelete, UBlueprint* BPContext = nullptr);
+	
+	/**
+	* Attempts to delete the given array of subobjects from their context 
+	*
+	* @param ContextHandle			The owning context of the subobjects that should be removed
+	* @param SubobjectsToDelete		Array of subobject handles that should be deleted
+	* @param OutComponentToSelect	Populates this handle with a valid selection in the component heirarchy if desiered
+	* @return 	The number of subobjects successfully deleted
+	*/
+	int32 DeleteSubobjects(const FSubobjectDataHandle& ContextHandle, const TArray<FSubobjectDataHandle>& SubobjectsToDelete, FSubobjectDataHandle& OutComponentToSelect, UBlueprint* BPContext = nullptr);
+	
+	/**
+	* Attempts to rename the given subobject to the new name.
+	*
+	* @param Handle			Handle to the subobject to rename
+	* @param InNewName		The new name that is desired for the give subobject
+	*
+	* @return True if the rename was successful, false otherwise.
+	*/
+	UFUNCTION(BlueprintCallable, Category = "SubobjectDataSubsystem")
+	bool RenameSubobject(const FSubobjectDataHandle& Handle, const FText& InNewName);
+
+	/**
+	* Attempts to reparent the given subobject to the new parent
+	*
+	* @param NewParentHandle		Handle of the new parent
+	* @param ToReparentHandle		The handle of the subobject that will get moved
+	*
+	* @return True if the reparent was successful, false otherwise.
+	*/
+	UFUNCTION(BlueprintCallable, Category = "SubobjectDataSubsystem")
+    bool ReparentSubobject(const FReparentSubobjectParams& Params, const FSubobjectDataHandle& ToReparentHandle);
+
+	UFUNCTION(BlueprintCallable, Category = "SubobjectDataSubsystem")
+	bool MakeNewSceneRoot(const FSubobjectDataHandle& Context, const FSubobjectDataHandle& NewSceneRoot, UBlueprint* BPContext);
+	
+	/**
+	* Attempts to reparent all subobjects in the HandlesToMove array to the new parent handle.
+	* 
+	* @param NewParentHandle		Handle of the new parent
+	*/
+	UFUNCTION(BlueprintCallable, Category = "SubobjectDataSubsystem")
+	bool ReparentSubobjects(const FReparentSubobjectParams& Params, const TArray<FSubobjectDataHandle>& HandlesToMove);
+
+	/**
+	 * Remove the child subobject from the owner
+	 *
+	 * @return True if the child was successfully removed. 
+	 */
+	UFUNCTION(BlueprintCallable, Category = "SubobjectDataSubsystem")
+	bool DetachSubobject(const FSubobjectDataHandle& OwnerHandle, const FSubobjectDataHandle& ChildToRemove);
+
+	/**
+	 * Addd the given subobject to a new owner. This will remove the subobject from its previous 
+	 * owner if necessary.
+	 *
+	 * @param OwnerHandle		The new owner to attach to
+	 * @param ChildToAddHandle	Handle to the subobject that will become a child of the owner
+	 * 
+	 * @return true if the child was added successfully
+	 */
+	UFUNCTION(BlueprintCallable, Category = "SubobjectDataSubsystem")
+    bool AttachSubobject(const FSubobjectDataHandle& OwnerHandle, const FSubobjectDataHandle& ChildToAddHandle);
+	
+	/**
+	 * Returns true if the given new text is a valid option to rename the
+	 * subobject with the given handle. Populates the OutErrorMessage if
+	 * it is not valid.
+	 *
+	 * @param Handle			Handle to the subobject that is being checked
+	 * @param InNewText			The new name that is desired
+	 * @param OutErrorMessage	The reasoning for an invalid name
+	 *
+	 * @return True if the rename is valid
+	 */
+	UFUNCTION(BlueprintCallable, Category = "SubobjectDataSubsystem")
+	bool IsValidRename(const FSubobjectDataHandle& Handle, const FText& InNewText, FText& OutErrorMessage) const;
+
+	/**
+	 * Returns true if the given array of handles represents subobjects that
+	 * can be copied.
+	 */
+	UFUNCTION(BlueprintCallable, Category = "SubobjectDataSubsystem")
+	bool CanCopySubobjects(const TArray<FSubobjectDataHandle>& Handles) const;
+
+	UFUNCTION(BlueprintCallable, Category = "SubobjectDataSubsystem")
+	void CopySubobjects(const TArray<FSubobjectDataHandle>& Handles, UBlueprint* BpContext);
+
+	UFUNCTION(BlueprintCallable, Category = "SubobjectDataSubsystem")
+	bool CanPasteSubobjects(const FSubobjectDataHandle& RootHandle, UBlueprint* BPContext = nullptr) const;
+
+	void PasteSubobjects(const FSubobjectDataHandle& PasteToContext, const TArray<FSubobjectDataHandle>& NewParentHandles, UBlueprint* BpContext, TArray<FSubobjectDataHandle>& OutPastedHandles);
+	
+	/**
+	 * Duplicate the given array of subobjects on the context. 
+	 *
+	 * @param Context			The owning context that the subobjects to dup come from
+	 * @param SubobjectsToDup	Array of handles of existing subobejcts you would like to have duplicated
+	 * @param BpContext			Pointer to the current blueprint context if necessary. Use nullptr if dealing with instances
+	 */
+	void DuplicateSubobjects(const FSubobjectDataHandle& Context, const TArray<FSubobjectDataHandle>& SubobjectsToDup, UBlueprint* BpContext);
+	
+	FScopedTransaction* BeginTransaction(const TArray<FSubobjectDataHandle>& Handles, const FText& Description, UBlueprint* Blueprint);
+
+	/**
+	* Find the scene root for a given subobject handle.
+	* If this is an actor, then it will return the handle pointing to its DefaultSceneRoot.
+	* If this is a component, it will walk the parent chain until it finds the DefaultSceneRoot.
+	*/
+	FSubobjectDataHandle FindSceneRootForSubobject(const FSubobjectDataHandle& InHandle) const;
+
+	UFUNCTION(BlueprintCallable, Category = "SubobjectDataSubsystem")
+	static void RenameSubobjectMemberVariable(UBlueprint* BPContext, const FSubobjectDataHandle& InHandle, const FName NewName);
+	
+private:
+
+	/**
+	 * Find all child subobject data related to the given root context.
+	 * Populates the OutVisited set with all data. 
+	 */
+	void FindAllSubobjectData(FSubobjectData* RootContext, TSet<FSubobjectData*>& OutVisited) const;
+	
+	/**
+	 * Gets a ref to subobject data that has been initalized with the given context.
+	 * It will have a unique ID, but may be reused from the FreeData array. 
+	 */
+	FSubobjectDataHandle CreateSubobjectData(UObject* Context, const FSubobjectDataHandle& ParentHandle = FSubobjectDataHandle::InvalidHandle, bool bIsInheritedSCS = false);
+
+	/**
+	 * Get a handle to the subobject with the given context with a specific parent.
+	 * If the context is already a known child of the parent, then it will return an existing subobject handle. 
+	 */
+	FSubobjectDataHandle FactoryCreateSubobjectDataWithParent(UObject* Context, const FSubobjectDataHandle& ParentHandle, bool bIsInheritedSCS = false);
+
+	/**
+	 * Create a subobject data and handle for an inherited blueprint context.
+	 * This will populate the given array of handles recursively with all possible children
+	 * (which are USCS_Node* in this case)
+	 */
+	FSubobjectDataHandle FactoryCreateInheritedBpSubobject(UObject* Context, const FSubobjectDataHandle& ParentHandle, bool bIsInherited, TArray<FSubobjectDataHandle>& OutArray);
+	
+	/** Find the parent that this new subobject should attach to */
+	FSubobjectDataHandle FindParentForNewSubobject(const UObject* NewSubobject, const FSubobjectDataHandle& SelectedParent);
+	
+	FSubobjectFactoryManager* FactoryManager = nullptr;
+	
+public:
+	FSubobjectFactoryManager* GetSubobjectFactoryManager() const { return FactoryManager; }
+	
+};
\ No newline at end of file
diff --git a/Engine/Source/Editor/SubobjectDataInterface/SubobjectDataInterface.Build.cs b/Engine/Source/Editor/SubobjectDataInterface/SubobjectDataInterface.Build.cs
new file mode 100644
--- /dev/null
+++ b/Engine/Source/Editor/SubobjectDataInterface/SubobjectDataInterface.Build.cs
@@ -0,0 +1,41 @@
+// Copyright Epic Games, Inc. All Rights Reserved.
+
+using UnrealBuildTool;
+using System.IO;
+
+public class SubobjectDataInterface : ModuleRules
+{
+	public SubobjectDataInterface(ReadOnlyTargetRules Target) : base(Target)
+	{
+		PublicIncludePaths.Add(Path.Combine(ModuleDirectory, "Public"));
+		PrivateIncludePaths.Add(Path.Combine(ModuleDirectory, "Private"));
+
+		// #TODO_BH Remove UnrealEd as a dependency here, because that sucks.
+		PublicDependencyModuleNames.AddRange(
+			new string[]
+			{
+				"Core",
+				"CoreUObject",
+			}
+		);
+
+		PrivateDependencyModuleNames.AddRange(
+			new string[]
+			{
+				"BlueprintGraph",
+				"Engine",
+			}
+		);
+
+		if (Target.bBuildEditor == true)
+		{
+			PrivateDependencyModuleNames.AddRange(
+				new string[]
+				{
+					"UnrealEd",	// ComponentEditorUtils, AssetBrokerage
+					"GameProjectGeneration",
+				}
+			);
+		}
+	}
+}
diff --git a/Engine/Source/Runtime/Engine/Classes/EdGraph/EdGraphPin.h b/Engine/Source/Runtime/Engine/Classes/EdGraph/EdGraphPin.h
--- a/Engine/Source/Runtime/Engine/Classes/EdGraph/EdGraphPin.h
+++ b/Engine/Source/Runtime/Engine/Classes/EdGraph/EdGraphPin.h
@@ -86,7 +86,7 @@ inline bool operator==(const FEdGraphTerminalType& A, const FEdGraphTerminalType
 }
 
 /** Struct used to define the type of information carried on this pin */
-USTRUCT()
+USTRUCT(BlueprintType)
 struct FEdGraphPinType
 {
 	GENERATED_USTRUCT_BODY()
