Fix PVS CIS issues

#rnx

diff --git a/Engine/Source/Editor/SubobjectDataInterface/Private/SubobjectDataSubsystem.cpp b/Engine/Source/Editor/SubobjectDataInterface/Private/SubobjectDataSubsystem.cpp
--- a/Engine/Source/Editor/SubobjectDataInterface/Private/SubobjectDataSubsystem.cpp
+++ b/Engine/Source/Editor/SubobjectDataInterface/Private/SubobjectDataSubsystem.cpp
@@ -1749,7 +1749,7 @@ void USubobjectDataSubsystem::PasteSubobjects(const FSubobjectDataHandle& PasteT
 			
 			// Create a new SCS node to contain the new component and add it to the tree
 			USCS_Node* NewSCSNode = Blueprint->SimpleConstructionScript->CreateNodeAndRenameComponent(NewActorComponent);
-			NewActorComponent = NewSCSNode ? NewSCSNode->ComponentTemplate : nullptr;
+			NewActorComponent = NewSCSNode ? ToRawPtr(NewSCSNode->ComponentTemplate) : nullptr;
 
 			FSubobjectDataHandle TargetParentHandle = FindParentForNewSubobject(NewActorComponent, PasteToContext);
 			TSharedPtr<FInheritedSubobjectData> TargetData = StaticCastSharedPtr<FInheritedSubobjectData>(TargetParentHandle.GetSharedDataPtr());
