Remove un-needed comment

#rnx

diff --git a/Engine/Source/Editor/SubobjectDataInterface/SubobjectDataInterface.Build.cs b/Engine/Source/Editor/SubobjectDataInterface/SubobjectDataInterface.Build.cs
--- a/Engine/Source/Editor/SubobjectDataInterface/SubobjectDataInterface.Build.cs
+++ b/Engine/Source/Editor/SubobjectDataInterface/SubobjectDataInterface.Build.cs
@@ -10,7 +10,6 @@ public class SubobjectDataInterface : ModuleRules
 		PublicIncludePaths.Add(Path.Combine(ModuleDirectory, "Public"));
 		PrivateIncludePaths.Add(Path.Combine(ModuleDirectory, "Private"));
 
-		// #TODO_BH Remove UnrealEd as a dependency here, because that sucks.
 		PublicDependencyModuleNames.AddRange(
 			new string[]
 			{
