Update comments on CreateNewCPPComponent and expose it to BP

#rnx

diff --git a/Engine/Source/Editor/SubobjectDataInterface/Public/SubobjectDataSubsystem.h b/Engine/Source/Editor/SubobjectDataInterface/Public/SubobjectDataSubsystem.h
--- a/Engine/Source/Editor/SubobjectDataInterface/Public/SubobjectDataSubsystem.h
+++ b/Engine/Source/Editor/SubobjectDataInterface/Public/SubobjectDataSubsystem.h
@@ -166,9 +166,7 @@ public:
 	 *
 	 * @return The new class that was created
 	 */
-	// #TODO_BH Finish up this function before exposing it to BP, it has some kinks to
-	// work out related to loading the newly created class.
-	//UFUNCTION(BlueprintCallable, Category = "SubobjectDataSubsystem")
+	UFUNCTION(BlueprintCallable, Category = "SubobjectDataSubsystem")
 	static UClass* CreateNewCPPComponent(TSubclassOf<UActorComponent> ComponentClass, const FString& NewClassPath, const FString& NewClassName);
 
 	/**
