Correctly return true for CanEdit() on a FSubobjectData that is an instance-added component, such as the default scene root on an empty actor.

#rnx

diff --git a/Engine/Source/Editor/SubobjectDataInterface/Private/SubobjectData.cpp b/Engine/Source/Editor/SubobjectDataInterface/Private/SubobjectData.cpp
--- a/Engine/Source/Editor/SubobjectDataInterface/Private/SubobjectData.cpp
+++ b/Engine/Source/Editor/SubobjectDataInterface/Private/SubobjectData.cpp
@@ -44,6 +44,15 @@ bool FSubobjectData::CanEdit() const
 	{
 		return true;
 	}
+	// If this is an instance-added component, then we can edit it. We know this isn't inherted because it would
+	// be a FInheritedSubobjectData
+	else if(const UActorComponent* Component = GetComponentTemplate())
+	{
+		if(Component->CreationMethod == EComponentCreationMethod::Instance)
+		{
+			return true;
+		}
+	}
 	return false;
 }
 
