Fix Subobject Data Factories not creating inherited subobject data in situations where a USCS_Node was added via the SimpleConstruction Script

#rnx

diff --git a/Engine/Source/Editor/SubobjectDataInterface/Private/InheritedSubobjectDataFactory.cpp b/Engine/Source/Editor/SubobjectDataInterface/Private/InheritedSubobjectDataFactory.cpp
--- a/Engine/Source/Editor/SubobjectDataInterface/Private/InheritedSubobjectDataFactory.cpp
+++ b/Engine/Source/Editor/SubobjectDataInterface/Private/InheritedSubobjectDataFactory.cpp
@@ -2,6 +2,7 @@
 
 #include "InheritedSubobjectDataFactory.h"
 #include "InheritedSubobjectData.h"
+#include "Engine/SCS_Node.h"		// #TODO_BH  We need to remove this when the actual subobject refactor happens
 
 TSharedPtr<FSubobjectData> FInheritedSubobjectDataFactory::CreateSubobjectData(const FCreateSubobjectParams& Params)
 {
@@ -13,11 +14,18 @@ bool FInheritedSubobjectDataFactory::ShouldCreateSubobjectData(const FCreateSubo
 	if(UActorComponent* Component = Cast<UActorComponent>(Params.Context))
 	{
 		// Create an inherited subobject data
-		if(Params.bIsInheritedSCS || Component->CreationMethod == EComponentCreationMethod::Native)
+		if(Params.bIsInheritedSCS || Component->CreationMethod != EComponentCreationMethod::Instance)
 		{
 			return true;
 		}
 	}
+	else if (USCS_Node* SCS = Cast<USCS_Node>(Params.Context))
+	{
+		if (UActorComponent* Comp = SCS->ComponentTemplate)
+		{
+			return Comp->CreationMethod != EComponentCreationMethod::Instance;
+		}
+	}
 
 	return false;
 }
\ No newline at end of file
