Add cateogry meta tags to the subobject data param structs.

#rnx

diff --git a/Engine/Source/Editor/SubobjectDataInterface/Public/SubobjectDataSubsystem.h b/Engine/Source/Editor/SubobjectDataInterface/Public/SubobjectDataSubsystem.h
--- a/Engine/Source/Editor/SubobjectDataInterface/Public/SubobjectDataSubsystem.h
+++ b/Engine/Source/Editor/SubobjectDataInterface/Public/SubobjectDataSubsystem.h
@@ -13,7 +13,7 @@ class UActorComponent;
 class FScopedTransaction;
 
 /** Options when adding a new subobject */
-USTRUCT(BlueprintType)
+USTRUCT(BlueprintType, Category = "SubobjectDataSubsystem")
 struct FAddNewSubobjectParams
 {
 	GENERATED_USTRUCT_BODY()
@@ -28,11 +28,11 @@ struct FAddNewSubobjectParams
 	{
 	}
 	
-	UPROPERTY(BlueprintReadWrite, Category = "SubobjectDataSubsystem|AddNewSubobject")
+	UPROPERTY(BlueprintReadWrite, Category = "SubobjectDataSubsystem")
 	FSubobjectDataHandle ParentHandle;
 
 	/** The class of the new subobject that will be added */
-	UPROPERTY(BlueprintReadWrite, Category = "SubobjectDataSubsystem|AddNewSubobject")
+	UPROPERTY(BlueprintReadWrite, Category = "SubobjectDataSubsystem")
 	UClass* NewClass;
 
 	/** Specific asset to use instead of the selected asset in the content browser */
@@ -42,20 +42,20 @@ struct FAddNewSubobjectParams
 	 * Pointer to the blueprint context that this subobject is in. If this is null, it is assumed that
 	 * this subobject is being added to an instance.
 	 */
-	UPROPERTY(BlueprintReadWrite, Category = "SubobjectDataSubsystem|AddNewSubobject")
+	UPROPERTY(BlueprintReadWrite, Category = "SubobjectDataSubsystem")
 	UBlueprint* BlueprintContext;
 	
 	/** Optionally skip marking this blueprint as modified (e.g. if we're handling that externally */
-	UPROPERTY(BlueprintReadWrite, Category = "SubobjectDataSubsystem|AddNewSubobject")
+	UPROPERTY(BlueprintReadWrite, Category = "SubobjectDataSubsystem")
 	uint8 bSkipMarkBlueprintModified : 1;
 
 	/** Whether the newly created component should keep its transform, or conform it to its parent */
-	UPROPERTY(BlueprintReadWrite, Category = "SubobjectDataSubsystem|AddNewSubobject")
+	UPROPERTY(BlueprintReadWrite, Category = "SubobjectDataSubsystem")
 	uint8 bConformTransformToParent : 1;
 };
 
 /** Options for reparenting subobjects */
-USTRUCT(BlueprintType)
+USTRUCT(BlueprintType, Category = "SubobjectDataSubsystem")
 struct FReparentSubobjectParams
 {
 	GENERATED_USTRUCT_BODY()
@@ -63,21 +63,21 @@ struct FReparentSubobjectParams
 	FReparentSubobjectParams() = default;
 
 	/** The handle of the subobject to reparent to. */
-	UPROPERTY(BlueprintReadWrite, Category = "SubobjectDataSubsystem|AddNewSubobject")
+	UPROPERTY(BlueprintReadWrite, Category = "SubobjectDataSubsystem")
 	FSubobjectDataHandle NewParentHandle;
 
 	/**
 	* Pointer to the blueprint context that this subobject is in. If this is null, it is assumed that
 	* this subobject is being added to an instance.
 	*/
-	UPROPERTY(BlueprintReadWrite, Category = "SubobjectDataSubsystem|AddNewSubobject")
+	UPROPERTY(BlueprintReadWrite, Category = "SubobjectDataSubsystem")
 	UBlueprint* BlueprintContext = nullptr;
 
 	/** 
 	* The preview actor context to be used if in a blueprint context. 
 	* This must have a value if BlueprintContext is needed. 
 	*/
-	UPROPERTY(BlueprintReadWrite, Category = "SubobjectDataSubsystem|AddNewSubobject")
+	UPROPERTY(BlueprintReadWrite, Category = "SubobjectDataSubsystem")
 	AActor* ActorPreviewContext = nullptr;
 };
 
@@ -87,7 +87,7 @@ struct FReparentSubobjectParams
 * instance clicked on via the level editor, a UBlueprint* by opening an asset,
 * or something piped in from python or other scripting languages. 
 */
-UCLASS()
+UCLASS(Category = "SubobjectDataSubsystem")
 class SUBOBJECTDATAINTERFACE_API USubobjectDataSubsystem : public UEngineSubsystem
 {
 	GENERATED_BODY()
