Fix incorrect color tinting of icons in inherited blueprints

#rnx

diff --git a/Engine/Source/Editor/SubobjectDataInterface/Private/SubobjectData.cpp b/Engine/Source/Editor/SubobjectDataInterface/Private/SubobjectData.cpp
--- a/Engine/Source/Editor/SubobjectDataInterface/Private/SubobjectData.cpp
+++ b/Engine/Source/Editor/SubobjectDataInterface/Private/SubobjectData.cpp
@@ -356,7 +356,7 @@ FText FSubobjectData::GetDisplayNameContextModifiers(bool bShowNativeComponentNa
 		if (bHasValidVarName && !bIsArrayVariable)
 		{
 			const bool bIsNative = IsNativeComponent();
-			const bool bIsInherited = IsInheritedComponent();
+			const bool bIsInherited = IsInheritedSCSNode();
 			const bool bIsInstance = IsInstancedComponent();
 
 			if ((bIsNative || bIsInherited) && bShowNativeComponentNames)
@@ -784,7 +784,7 @@ FLinearColor FSubobjectData::GetColorTintForIcon() const
 		{
 			return InstancedInheritedBlueprintComponentColor;
 		}
-		else
+		else if(IsInheritedSCSNode())
 		{
 			return InheritedBlueprintComponentColor;
 		}
