Remove unused variables in SubobjectData.cpp

#rnx

diff --git a/Engine/Source/Editor/SubobjectDataInterface/Private/SubobjectData.cpp b/Engine/Source/Editor/SubobjectDataInterface/Private/SubobjectData.cpp
--- a/Engine/Source/Editor/SubobjectDataInterface/Private/SubobjectData.cpp
+++ b/Engine/Source/Editor/SubobjectDataInterface/Private/SubobjectData.cpp
@@ -1,7 +1,6 @@
 // Copyright Epic Games, Inc. All Rights Reserved.
 
 #include "SubobjectData.h"
-#include "SubobjectDataSubsystem.h"		// For iterating child if we need to
 #include "Engine/Blueprint.h"			// Casting to UBlueprint
 #include "Components/ChildActorComponent.h"
 #include "GameFramework/Actor.h"
@@ -510,16 +509,13 @@ FSubobjectDataHandle FSubobjectData::FindChildByObject(UObject* ContextObject) c
 		return FSubobjectDataHandle::InvalidHandle;
 	}
 
-	if (USubobjectDataSubsystem* System = USubobjectDataSubsystem::Get())
+	for (const FSubobjectDataHandle& CurrentChild : ChildrenHandles)
 	{
-		for (const FSubobjectDataHandle& CurrentChild : ChildrenHandles)
+		if (FSubobjectData* ChildData = CurrentChild.GetData())
 		{
-			if (FSubobjectData* ChildData = CurrentChild.GetData())
+			if (ChildData->GetObject() == ContextObject)
 			{
-				if (ChildData->GetObject() == ContextObject)
-				{
-					return CurrentChild;
-				}
+				return CurrentChild;
 			}
 		}
 	}
@@ -939,7 +935,6 @@ bool FSubobjectData::IsInstancedInheritedComponent() const
 		return false;
 	}
 	FSubobjectDataHandle CurrentHandle = ParentObjectHandle;
-	USubobjectDataSubsystem* System = USubobjectDataSubsystem::Get();
 	FSubobjectData* CurrentData = CurrentHandle.GetData();
 			
 	while(CurrentHandle.IsValid() && CurrentData && !CurrentData->IsActor())
@@ -953,19 +948,16 @@ bool FSubobjectData::IsInstancedInheritedComponent() const
 
 bool FSubobjectData::IsAttachedTo(const FSubobjectDataHandle& InHandle) const
 {
-	if (USubobjectDataSubsystem* System = USubobjectDataSubsystem::Get())
-	{
-		FSubobjectDataHandle TestParentHandle = ParentObjectHandle;
+	FSubobjectDataHandle TestParentHandle = ParentObjectHandle;
 
-		while (TestParentHandle.IsValid())
+	while (TestParentHandle.IsValid())
+	{
+		if (TestParentHandle == InHandle)
 		{
-			if (TestParentHandle == InHandle)
-			{
-				return true;
-			}
-			const FSubobjectData* TestParentData = TestParentHandle.GetData();
-			TestParentHandle = TestParentData ? TestParentData->GetParentHandle() : FSubobjectDataHandle::InvalidHandle;
+			return true;
 		}
+		const FSubobjectData* TestParentData = TestParentHandle.GetData();
+		TestParentHandle = TestParentData ? TestParentData->GetParentHandle() : FSubobjectDataHandle::InvalidHandle;
 	}
 
 	return false;
