Correctly add "(Inherited)" in the case of a native component in the instance editor. Do not display the native component name for instance inhertied bluerpint subobjects.

#rnx

diff --git a/Engine/Source/Editor/SubobjectDataInterface/Private/SubobjectData.cpp b/Engine/Source/Editor/SubobjectDataInterface/Private/SubobjectData.cpp
--- a/Engine/Source/Editor/SubobjectDataInterface/Private/SubobjectData.cpp
+++ b/Engine/Source/Editor/SubobjectDataInterface/Private/SubobjectData.cpp
@@ -257,7 +257,7 @@ FString FSubobjectData::GetDisplayString(bool bShowNativeComponentNames /* = tru
 			FStringFormatNamedArguments Args;
 			Args.Add(TEXT("VarName"), VariableProperty && VariableProperty->IsNative() ? VariableProperty->GetDisplayNameText().ToString() : VariableName.ToString());
 			FString CompName = TEXT(" (") + ComponentTemplate->GetName() + TEXT(")");
-			Args.Add(TEXT("CompName"), bIsNative || IsInstancedInheritedComponent() ? CompName : TEXT(""));
+			Args.Add(TEXT("CompName"), bIsNative ? CompName : TEXT(""));
 			return FString::Format(TEXT("{VarName}{CompName}"), Args);
 		}
 		else
@@ -351,17 +351,16 @@ FText FSubobjectData::GetDisplayNameContextModifiers(bool bShowNativeComponentNa
             return FText::FromString(FString::Format(TEXT("{InheritedText}"), Args));
 		}
 		
-		// Only display SCS node variable names in the tree if they have not been autogenerated
+		// Only display node variable names in the tree if they have not been autogenerated
 		if (bHasValidVarName && !bIsArrayVariable)
 		{
 			const bool bIsNative = IsNativeComponent();
 			const bool bIsInherited = IsInheritedSCSNode();
-			const bool bIsInstance = IsInstancedComponent();
 
-			if ((bIsNative || bIsInherited) && bShowNativeComponentNames)
+			if ((bIsNative || bIsInherited || bIsBlueprintInstanceInherited) && bShowNativeComponentNames)
 			{		
 				FStringFormatNamedArguments Args;
-				Args.Add(TEXT("InheritedText"), bIsInherited || bIsBlueprintInstanceInherited ? TEXT("(Inherited)") : TEXT(""));
+				Args.Add(TEXT("InheritedText"), TEXT("(Inherited)"));
 				return FText::FromString(FString::Format(TEXT("{InheritedText}"), Args));
 			}
 		}
