Only set a newly created component as the default scene root if it can be reparented. This stops the incorrect ordering of native components in the heirarchy after a new component has been added

#rnx

diff --git a/Engine/Source/Editor/SubobjectDataInterface/Private/SubobjectDataSubsystem.cpp b/Engine/Source/Editor/SubobjectDataInterface/Private/SubobjectDataSubsystem.cpp
--- a/Engine/Source/Editor/SubobjectDataInterface/Private/SubobjectDataSubsystem.cpp
+++ b/Engine/Source/Editor/SubobjectDataInterface/Private/SubobjectDataSubsystem.cpp
@@ -751,7 +751,7 @@ FSubobjectDataHandle USubobjectDataSubsystem::AddNewSubobject(const FAddNewSubob
 		{
 			UActorComponent* NewComponent = FComponentEditorUtils::DuplicateComponent(ComponentTemplate);
 			// Create a new subobject data set with this component
-			NewDataHandle = FactoryCreateSubobjectDataWithParent(NewComponent, ParentObjHandle);
+			NewDataHandle = FactoryCreateSubobjectDataWithParent(NewComponent, ParentObjData->GetHandle());
 		}
 		else if(AActor* ActorInstance = ParentObjData->GetMutableActorContext())
 		{
@@ -778,7 +778,7 @@ FSubobjectDataHandle USubobjectDataSubsystem::AddNewSubobject(const FAddNewSubob
 			// Do Scene Attachment if this new Component is a USceneComponent
 			if (USceneComponent* NewSceneComponent = Cast<USceneComponent>(NewInstanceComponent))
 			{
-				if(ParentObjData->IsDefaultSceneRoot())
+				if(ParentObjData->IsDefaultSceneRoot() && ParentObjData->CanReparent())
 				{
 					ActorInstance->SetRootComponent(NewSceneComponent);
 				}
@@ -834,7 +834,7 @@ FSubobjectDataHandle USubobjectDataSubsystem::AddNewSubobject(const FAddNewSubob
 			if (!NewInstanceComponent->IsPendingKill())
 			{
 				// Create a new subobject data set with this component
-				NewDataHandle = FactoryCreateSubobjectDataWithParent(NewInstanceComponent, ParentObjHandle);
+				NewDataHandle = FactoryCreateSubobjectDataWithParent(NewInstanceComponent, ParentObjData->GetHandle());
 			}
 		}
 		else
