Correctly check for a root component that is native with a different name than the default. This will allow other systems such as drag/drop to evaulate correctly when the subobject editor is asked to find the default scene root given only the actor instance.

#rnx

diff --git a/Engine/Source/Editor/SubobjectDataInterface/Private/SubobjectData.cpp b/Engine/Source/Editor/SubobjectDataInterface/Private/SubobjectData.cpp
--- a/Engine/Source/Editor/SubobjectDataInterface/Private/SubobjectData.cpp
+++ b/Engine/Source/Editor/SubobjectDataInterface/Private/SubobjectData.cpp
@@ -885,7 +885,10 @@ bool FSubobjectData::IsDefaultSceneRoot() const
 	const USceneComponent* SceneComponent = Cast<USceneComponent>(GetComponentTemplate());
 	if (SceneComponent && !SceneComponent->IsTemplate())
 	{
-		return SceneComponent->GetFName() == USceneComponent::GetDefaultSceneRootVariableName();
+		const AActor* OwningActor = SceneComponent->GetAttachmentRootActor();		
+		return
+			SceneComponent->GetFName() == USceneComponent::GetDefaultSceneRootVariableName() || 
+			(OwningActor && SceneComponent == OwningActor->GetRootComponent());
 	}
 	// If this isn't a scene component, then we can check an SCS node for a flag. 
 	else if (const USCS_Node* SCS_Node = GetSCSNode())
