Fix a bug where the hierarchy would all be attached to the default scene root for native blueprints.
Correctly sort non-scene components to the bottom of the component tree.
This is some logic that was missed during the SCS editor refactor and is known to work already.

#rnx

diff --git a/Engine/Source/Editor/SubobjectDataInterface/Private/SubobjectDataSubsystem.cpp b/Engine/Source/Editor/SubobjectDataInterface/Private/SubobjectDataSubsystem.cpp
--- a/Engine/Source/Editor/SubobjectDataInterface/Private/SubobjectDataSubsystem.cpp
+++ b/Engine/Source/Editor/SubobjectDataInterface/Private/SubobjectDataSubsystem.cpp
@@ -77,6 +77,75 @@ void USubobjectDataSubsystem::K2_GatherSubobjectDataForInstance(AActor* Context,
 	return GatherSubobjectData(Context, OutArray);
 }
 
+FSubobjectDataHandle USubobjectDataSubsystem::FindOrCreateAttachParentForComponent(UActorComponent* InActorComponent, const FSubobjectDataHandle& ActorRootHandle, TArray<FSubobjectDataHandle>& ExistingHandles)
+{
+	check(InActorComponent != nullptr);
+	const FSubobjectData* ActorRootData = ActorRootHandle.GetData();
+
+	USceneComponent* SceneComponent = Cast<USceneComponent>(InActorComponent);
+	if (SceneComponent == nullptr)
+	{
+		check(ActorRootData->IsValid());
+		check(ActorRootData->IsActor());
+		return ActorRootData->GetHandle();
+	}
+
+	auto FindExistingHandle = [&ExistingHandles](USceneComponent* ComponentToFind) -> FSubobjectDataHandle
+	{
+		for (const FSubobjectDataHandle& Handle : ExistingHandles)
+		{
+			const FSubobjectData* Data = Handle.GetData();
+			if (Data && Data->GetComponentTemplate() == ComponentToFind)
+			{
+				return Handle;
+			}
+		}
+		return FSubobjectDataHandle::InvalidHandle;
+	};
+
+	FSubobjectDataHandle ParentHandle;
+	if (SceneComponent->GetAttachParent() != nullptr)
+	{
+		// Attempt to find the parent node in the current tree
+		ParentHandle = FindExistingHandle(SceneComponent->GetAttachParent());
+		if (!ParentHandle.IsValid())
+		{
+			// If the actual attach parent wasn't found, attempt to find its archetype.
+			// This handles the BP editor case where we might add UCS component nodes taken
+			// from the preview actor instance, which are not themselves template objects.
+			ParentHandle = FindExistingHandle(Cast<USceneComponent>(SceneComponent->GetAttachParent()->GetArchetype()));
+			if (!ParentHandle.IsValid())
+			{
+				// Recursively add the parent handle to the tree if it does not exist yet
+				ParentHandle = FactoryCreateSubobjectDataWithParent(SceneComponent, FindOrCreateAttachParentForComponent(
+					SceneComponent->GetAttachParent(),
+					ActorRootHandle,
+					ExistingHandles)
+				);
+
+				if (ParentHandle.IsValid())
+				{
+					ExistingHandles.Add(ParentHandle);
+				}				
+			}
+		}
+	}
+
+	// Use the default scene root node as a backup
+	if (!ParentHandle.IsValid())
+	{
+		ParentHandle = FindSceneRootForSubobject(ActorRootData->GetHandle());
+	}
+
+	// The actor doesn't have a root component yet
+	if (!ParentHandle.IsValid())
+	{
+		ParentHandle = ActorRootData->GetHandle();
+	}
+
+	return ParentHandle;
+}
+
 void USubobjectDataSubsystem::GatherSubobjectData(UObject* Context, TArray<FSubobjectDataHandle>& OutArray)
 {
 	if (!Context)
@@ -114,7 +183,8 @@ void USubobjectDataSubsystem::GatherSubobjectData(UObject* Context, TArray<FSubo
 		if (RootComponent != nullptr)
 		{
 			Components.Remove(RootComponent);
-			RootComponentHandle = FactoryCreateSubobjectDataWithParent(RootComponent, RootActorHandle);
+			const FSubobjectDataHandle& ParentHandle = FindOrCreateAttachParentForComponent(RootComponent, RootActorHandle, OutArray);
+			RootComponentHandle = FactoryCreateSubobjectDataWithParent(RootComponent, ParentHandle);
 			OutArray.Add(RootComponentHandle);
 		}
 
@@ -123,24 +193,32 @@ void USubobjectDataSubsystem::GatherSubobjectData(UObject* Context, TArray<FSubo
 		// These components will all be attached to the root component if it exists
 		for (UActorComponent* Component : Components)
 		{
+			const FSubobjectDataHandle& ParentHandle = FindOrCreateAttachParentForComponent(Component, RootActorHandle, OutArray);
+
 			FSubobjectDataHandle NewComponentHandle = FactoryCreateSubobjectDataWithParent(
 					/* Context = */ Component,
-					/* ParentHandle = */ RootComponentHandle.IsValid() ? RootComponentHandle : RootActorHandle
+					/* ParentHandle = */ ParentHandle
 			);
+
 			ensureMsgf(NewComponentHandle.IsValid(), TEXT("Gathering of native components failed!"));
 			OutArray.Add(NewComponentHandle);
 		}
 
 		// If it's a Blueprint-generated class, also get the inheritance stack
-		TArray<UBlueprint*> ParentBPStack;
+		TArray<UBlueprintGeneratedClass*> ParentBPStack;
 		UBlueprint::GetBlueprintHierarchyFromClass(ActorContext->GetClass(), ParentBPStack);
 
+		UBlueprint* ActorBP = (ParentBPStack.Num() > 0 && ParentBPStack[0]) ? Cast<UBlueprint>(ParentBPStack[0]->ClassGeneratedBy) : nullptr;
+		ensure(ActorBP);
+
 		// Add the full SCS tree node hierarchy (including SCS nodes inherited from parent blueprints)
 		for (int32 StackIndex = ParentBPStack.Num() - 1; StackIndex >= 0; --StackIndex)
 		{
-			if (ParentBPStack[StackIndex]->SimpleConstructionScript != nullptr)
+			USimpleConstructionScript* ParentSCS = ParentBPStack[StackIndex] ? ParentBPStack[StackIndex]->SimpleConstructionScript.Get() : nullptr;
+
+			if (ParentSCS)
 			{
-				const TArray<USCS_Node*>& SCS_RootNodes = ParentBPStack[StackIndex]->SimpleConstructionScript->GetRootNodes();
+				const TArray<USCS_Node*>& SCS_RootNodes = ParentSCS->GetRootNodes();
 				for (int32 NodeIndex = 0; NodeIndex < SCS_RootNodes.Num(); ++NodeIndex)
 				{
 					USCS_Node* SCS_Node = SCS_RootNodes[NodeIndex];
@@ -152,8 +230,7 @@ void USubobjectDataSubsystem::GatherSubobjectData(UObject* Context, TArray<FSubo
 					// If this SCS node has a parent component, then add it to that
 					if (SCS_Node->ParentComponentOrVariableName != NAME_None)
 					{
-						USceneComponent* ParentComponent = SCS_Node->GetParentComponentTemplate(ParentBPStack[0]);
-						if (ParentComponent != nullptr)
+						if (USceneComponent* ParentComponent = SCS_Node->GetParentComponentTemplate(ActorBP))
 						{
 							// The parent component will already be in the out array if it is valid, so search for it
 							FSubobjectDataHandle ParentHandle;
@@ -187,13 +264,13 @@ void USubobjectDataSubsystem::GatherSubobjectData(UObject* Context, TArray<FSubo
 					{
 						// This call creates ICH override templates for the current Blueprint. Without this, the parent node
 						// search above can fail when attempting to match an inherited node in the tree via component template.
-						NewData->GetObjectForBlueprint(ParentBPStack[0]);
+						NewData->GetObjectForBlueprint(ActorBP);
 						for (FSubobjectDataHandle ChildHandle : NewData->GetChildrenHandles())
 						{
 							FSubobjectData* ChildData = ChildHandle.GetData();
 							if (ensure(ChildData != nullptr))
 							{
-								ChildData->GetObjectForBlueprint(ParentBPStack[0]);
+								ChildData->GetObjectForBlueprint(ActorBP);
 							}
 						}
 					}
diff --git a/Engine/Source/Editor/SubobjectDataInterface/Public/SubobjectDataSubsystem.h b/Engine/Source/Editor/SubobjectDataInterface/Public/SubobjectDataSubsystem.h
--- a/Engine/Source/Editor/SubobjectDataInterface/Public/SubobjectDataSubsystem.h
+++ b/Engine/Source/Editor/SubobjectDataInterface/Public/SubobjectDataSubsystem.h
@@ -317,6 +317,11 @@ public:
 	
 private:
 
+	/**
+	* Recursively add the attachment children of the given actor component
+	*/
+	FSubobjectDataHandle FindOrCreateAttachParentForComponent(UActorComponent* InActorComponent, const FSubobjectDataHandle& ActorRootHandle, TArray<FSubobjectDataHandle>& ExistingHandles);
+
 	/**
 	 * Find all child subobject data related to the given root context.
 	 * Populates the OutVisited set with all data. 
