Remove unnecessary check in SubobjectDataSubsystem

#rnx

diff --git a/Engine/Source/Editor/SubobjectDataInterface/Private/SubobjectDataSubsystem.cpp b/Engine/Source/Editor/SubobjectDataInterface/Private/SubobjectDataSubsystem.cpp
--- a/Engine/Source/Editor/SubobjectDataInterface/Private/SubobjectDataSubsystem.cpp
+++ b/Engine/Source/Editor/SubobjectDataInterface/Private/SubobjectDataSubsystem.cpp
@@ -690,8 +690,6 @@ FSubobjectDataHandle USubobjectDataSubsystem::FindParentForNewSubobject(const UO
 		{
 			TargetParentHandle = SelectedParent;
 		}
-
-		check(TargetParentHandle.IsValid() && TargetParentData->IsActor());
 	}
 	
 	return TargetParentHandle;
