Fix a case where an ensure could be triggered when adding a new component to a blueprint with BP inherited components

#rnx

diff --git a/Engine/Source/Editor/SubobjectDataInterface/Private/SubobjectDataSubsystem.cpp b/Engine/Source/Editor/SubobjectDataInterface/Private/SubobjectDataSubsystem.cpp
--- a/Engine/Source/Editor/SubobjectDataInterface/Private/SubobjectDataSubsystem.cpp
+++ b/Engine/Source/Editor/SubobjectDataInterface/Private/SubobjectDataSubsystem.cpp
@@ -146,6 +146,50 @@ FSubobjectDataHandle USubobjectDataSubsystem::FindOrCreateAttachParentForCompone
 	return ParentHandle;
 }
 
+FSubobjectDataHandle USubobjectDataSubsystem::FindAttachParentForInheritedComponent(UActorComponent* InActorComponent, const FSubobjectDataHandle& InStartData, UBlueprint* BP)
+{
+	FSubobjectDataHandle OutHandle = FSubobjectDataHandle::InvalidHandle;
+
+	if (InActorComponent)
+	{
+		if (const FSubobjectData* ActorRootData = InStartData.GetData())
+		{
+			// Check to see if the given component template matches the given tree node
+			// 
+			// For certain node types, GetOrCreateEditableComponentTemplate() will handle retrieving 
+			// the "OverridenComponentTemplate" which may be what we're looking for in some 
+			// cases; if not, then we fall back to just checking GetComponentTemplate()
+			if (ActorRootData->GetObjectForBlueprint(BP) == InActorComponent)
+			{
+				OutHandle = ActorRootData->GetHandle();
+			}
+			else if (ActorRootData->GetComponentTemplate() == InActorComponent)
+			{
+				OutHandle = ActorRootData->GetHandle();
+			}
+			else
+			{
+				// Recursively search for the node in our child set
+				OutHandle = ActorRootData->FindChildByObject(InActorComponent);
+				if (!OutHandle.IsValid())
+				{
+					const TArray<FSubobjectDataHandle>& ChildHandles = ActorRootData->GetChildrenHandles();
+					for (const FSubobjectDataHandle& Handle : ChildHandles)
+					{
+						OutHandle = FindAttachParentForInheritedComponent(InActorComponent, Handle, BP);
+						if (OutHandle.IsValid())
+						{
+							break;
+						}
+					}
+				}
+			}
+		}
+	}
+
+	return OutHandle;
+}
+
 void USubobjectDataSubsystem::GatherSubobjectData(UObject* Context, TArray<FSubobjectDataHandle>& OutArray)
 {
 	if (!Context)
@@ -232,23 +276,11 @@ void USubobjectDataSubsystem::GatherSubobjectData(UObject* Context, TArray<FSubo
 					{
 						if (USceneComponent* ParentComponent = SCS_Node->GetParentComponentTemplate(ActorBP))
 						{
-							// The parent component will already be in the out array if it is valid, so search for it
-							FSubobjectDataHandle ParentHandle;
-							for(const FSubobjectDataHandle& CurHandle : OutArray)
-							{
-								if(TSharedPtr<FSubobjectData> DataPtr = CurHandle.GetSharedDataPtr())
-								{
-									if (DataPtr->GetComponentTemplate() == ParentComponent)
-									{
-										ParentHandle = CurHandle;
-										break;
-									}
-								}
-							}
-							
+							FSubobjectDataHandle ParentHandle = FindAttachParentForInheritedComponent(ParentComponent, RootActorHandle, ActorBP);
+
 							if (ensure(ParentHandle.IsValid()))
 							{
-								NewHandle = FactoryCreateInheritedBpSubobject(SCS_Node, RootComponentHandle.IsValid() ? RootComponentHandle : RootActorHandle, /* bIsInherited = */ StackIndex > 0, OutArray);
+								NewHandle = FactoryCreateInheritedBpSubobject(SCS_Node, ParentHandle, /* bIsInherited = */ StackIndex > 0, OutArray);
 							}
 						}
 					}
diff --git a/Engine/Source/Editor/SubobjectDataInterface/Public/SubobjectDataSubsystem.h b/Engine/Source/Editor/SubobjectDataInterface/Public/SubobjectDataSubsystem.h
--- a/Engine/Source/Editor/SubobjectDataInterface/Public/SubobjectDataSubsystem.h
+++ b/Engine/Source/Editor/SubobjectDataInterface/Public/SubobjectDataSubsystem.h
@@ -322,6 +322,11 @@ private:
 	*/
 	FSubobjectDataHandle FindOrCreateAttachParentForComponent(UActorComponent* InActorComponent, const FSubobjectDataHandle& ActorRootHandle, TArray<FSubobjectDataHandle>& ExistingHandles);
 
+	/**
+	* Recursively search for an attachment parent that matches the given InActorComponent on the blueprint
+	*/
+	FSubobjectDataHandle FindAttachParentForInheritedComponent(UActorComponent* InActorComponent, const FSubobjectDataHandle& ActorRootHandle, UBlueprint* BP);
+
 	/**
 	 * Find all child subobject data related to the given root context.
 	 * Populates the OutVisited set with all data. 
