Sort blueprint subobjects by their type to ensure that non-scene components show up on the bottom. When adding a new subobject parent non-scene components to the actor instead of the default scene root.

diff --git a/Engine/Source/Editor/SubobjectDataInterface/Private/SubobjectDataSubsystem.cpp b/Engine/Source/Editor/SubobjectDataInterface/Private/SubobjectDataSubsystem.cpp
--- a/Engine/Source/Editor/SubobjectDataInterface/Private/SubobjectDataSubsystem.cpp
+++ b/Engine/Source/Editor/SubobjectDataInterface/Private/SubobjectDataSubsystem.cpp
@@ -309,6 +309,12 @@ void USubobjectDataSubsystem::GatherSubobjectData(UObject* Context, TArray<FSubo
 				}
 			}
 		}
+
+		// Sort subobjects by type (always put scene components first in the tree)
+		OutArray.Sort([](const FSubobjectDataHandle& A, const FSubobjectDataHandle& B)
+		{
+			return A.GetData()->IsActor() || (A.GetData()->IsSceneComponent() && !B.GetData()->IsActor());
+		});
 	}
 	// Otherwise, this is an actor instance in a level
 	else
@@ -714,14 +720,17 @@ FSubobjectDataHandle USubobjectDataSubsystem::FindParentForNewSubobject(const UO
 	}
 	else
 	{
-		if (TargetParentData->IsValid())
-		{
-			TargetParentHandle = FindSceneRootForSubobject(SelectedParent);
-		}
-		else
+		// Non-scene components should be parented to the base actor in the hierarchy
+		FSubobjectDataHandle CurrentHandle = SelectedParent;
+		FSubobjectData* CurrentData = CurrentHandle.GetData();
+
+		while (CurrentData && !CurrentData->IsActor())
 		{
-			TargetParentHandle = SelectedParent;
+			CurrentHandle = CurrentData->GetParentHandle();
+			CurrentData = CurrentHandle.GetData();
 		}
+
+		TargetParentHandle = CurrentHandle;
 	}
 	
 	return TargetParentHandle;
