When adding to instance inhertied subobjects, if the actor instance cannot be initalliy found walk up the heirarchy until we find one.

diff --git a/Engine/Source/Editor/SubobjectDataInterface/Private/SubobjectDataSubsystem.cpp b/Engine/Source/Editor/SubobjectDataInterface/Private/SubobjectDataSubsystem.cpp
--- a/Engine/Source/Editor/SubobjectDataInterface/Private/SubobjectDataSubsystem.cpp
+++ b/Engine/Source/Editor/SubobjectDataInterface/Private/SubobjectDataSubsystem.cpp
@@ -721,21 +721,26 @@ FSubobjectDataHandle USubobjectDataSubsystem::FindParentForNewSubobject(const UO
 	else
 	{
 		// Non-scene components should be parented to the base actor in the hierarchy
-		FSubobjectDataHandle CurrentHandle = SelectedParent;
-		FSubobjectData* CurrentData = CurrentHandle.GetData();
-
-		while (CurrentData && !CurrentData->IsActor())
-		{
-			CurrentHandle = CurrentData->GetParentHandle();
-			CurrentData = CurrentHandle.GetData();
-		}
-
-		TargetParentHandle = CurrentHandle;
+		TargetParentHandle = GetActorRootHandle(SelectedParent);
 	}
 	
 	return TargetParentHandle;
 }
 
+FSubobjectDataHandle USubobjectDataSubsystem::GetActorRootHandle(const FSubobjectDataHandle& StartingHandle)
+{
+	FSubobjectDataHandle CurrentHandle = StartingHandle;
+	FSubobjectData* CurrentData = CurrentHandle.GetData();
+
+	while (CurrentData && !CurrentData->IsActor())
+	{
+		CurrentHandle = CurrentData->GetParentHandle();
+		CurrentData = CurrentHandle.GetData();
+	}
+
+	return CurrentHandle;
+}
+
 FSubobjectDataHandle USubobjectDataSubsystem::AddNewSubobject(const FAddNewSubobjectParams& Params, FText& FailReason)
 {
 	FSubobjectDataHandle NewDataHandle = FSubobjectDataHandle::InvalidHandle;
@@ -869,94 +874,110 @@ FSubobjectDataHandle USubobjectDataSubsystem::AddNewSubobject(const FAddNewSubob
 			// Create a new subobject data set with this component
 			NewDataHandle = FactoryCreateSubobjectDataWithParent(NewComponent, ParentObjData->GetHandle());
 		}
-		else if(AActor* ActorInstance = ParentObjData->GetMutableActorContext())
+		else
 		{
-			ActorInstance->Modify();
+			AActor* ActorInstance = ParentObjData->GetMutableActorContext();
 
-			// Create an appropriate name for the new component
-			FName NewComponentName = NAME_None;
-			if (Asset)
+			// If we can't find an actor instance then search up the hierarchy for one that we can use instead
+			if (!ActorInstance)
 			{
-				NewComponentName = *FComponentEditorUtils::GenerateValidVariableNameFromAsset(Asset, ActorInstance);
-			}
-			else
-			{
-				NewComponentName = *FComponentEditorUtils::GenerateValidVariableName(NewClass, ActorInstance);
+				const FSubobjectDataHandle& ActorRootHandle = GetActorRootHandle(ParentObjData->GetHandle());
+				if (ActorRootHandle.IsValid())
+				{
+					ParentObjData = ActorRootHandle.GetData();
+					ActorInstance = ParentObjData ? ParentObjData->GetMutableActorContext() : nullptr;
+				}
 			}
-			
-			// Get the set of owned components that exists prior to instancing the new component.
-			TInlineComponentArray<UActorComponent*> PreInstanceComponents;
-			ActorInstance->GetComponents(PreInstanceComponents);
 
-			// Construct the new component and attach as needed
-			UActorComponent* NewInstanceComponent = NewObject<UActorComponent>(ActorInstance, NewClass, NewComponentName, RF_Transactional);
-			
-			// Do Scene Attachment if this new Component is a USceneComponent
-			if (USceneComponent* NewSceneComponent = Cast<USceneComponent>(NewInstanceComponent))
+			if (ActorInstance && ParentObjData)
 			{
-				if(ParentObjData->IsDefaultSceneRoot() && ParentObjData->CanReparent())
+				ActorInstance->Modify();
+
+				// Create an appropriate name for the new component
+				FName NewComponentName = NAME_None;
+				if (Asset)
 				{
-					ActorInstance->SetRootComponent(NewSceneComponent);
+					NewComponentName = *FComponentEditorUtils::GenerateValidVariableNameFromAsset(Asset, ActorInstance);
 				}
 				else
 				{
-					USceneComponent* AttachTo = Cast<USceneComponent>(ParentObjData->GetMutableComponentTemplate());
-					if (AttachTo == nullptr)
-					{
-						AttachTo = ActorInstance->GetRootComponent();
-					}
-					check(AttachTo != nullptr);
+					NewComponentName = *FComponentEditorUtils::GenerateValidVariableName(NewClass, ActorInstance);
+				}
+
+				// Get the set of owned components that exists prior to instancing the new component.
+				TInlineComponentArray<UActorComponent*> PreInstanceComponents;
+				ActorInstance->GetComponents(PreInstanceComponents);
 
-					// Make sure that the mobility of the new scene component is such that we can attach it
-					if (AttachTo->Mobility == EComponentMobility::Movable)
+				// Construct the new component and attach as needed
+				UActorComponent* NewInstanceComponent = NewObject<UActorComponent>(ActorInstance, NewClass, NewComponentName, RF_Transactional);
+
+				// Do Scene Attachment if this new Component is a USceneComponent
+				if (USceneComponent* NewSceneComponent = Cast<USceneComponent>(NewInstanceComponent))
+				{
+					if (ParentObjData->IsDefaultSceneRoot() && ParentObjData->CanReparent())
 					{
-						NewSceneComponent->Mobility = EComponentMobility::Movable;
+						ActorInstance->SetRootComponent(NewSceneComponent);
 					}
-					else if (AttachTo->Mobility == EComponentMobility::Stationary && NewSceneComponent->Mobility == EComponentMobility::Static)
+					else
 					{
-						NewSceneComponent->Mobility = EComponentMobility::Stationary;
-					}
+						USceneComponent* AttachTo = Cast<USceneComponent>(ParentObjData->GetMutableComponentTemplate());
+						if (AttachTo == nullptr)
+						{
+							AttachTo = ActorInstance->GetRootComponent();
+						}
+						check(AttachTo != nullptr);
 
-					NewSceneComponent->AttachToComponent(AttachTo, FAttachmentTransformRules::SnapToTargetNotIncludingScale);
+						// Make sure that the mobility of the new scene component is such that we can attach it
+						if (AttachTo->Mobility == EComponentMobility::Movable)
+						{
+							NewSceneComponent->Mobility = EComponentMobility::Movable;
+						}
+						else if (AttachTo->Mobility == EComponentMobility::Stationary && NewSceneComponent->Mobility == EComponentMobility::Static)
+						{
+							NewSceneComponent->Mobility = EComponentMobility::Stationary;
+						}
+
+						NewSceneComponent->AttachToComponent(AttachTo, FAttachmentTransformRules::SnapToTargetNotIncludingScale);
+					}
 				}
-			}
 
-			// If the component was created from/for a particular asset, assign it now
-			if (Asset)
-			{
-				FComponentAssetBrokerage::AssignAssetToComponent(NewInstanceComponent, Asset);
-			}
+				// If the component was created from/for a particular asset, assign it now
+				if (Asset)
+				{
+					FComponentAssetBrokerage::AssignAssetToComponent(NewInstanceComponent, Asset);
+				}
 
-			// Add to SerializedComponents array so it gets saved
-			ActorInstance->AddInstanceComponent(NewInstanceComponent);
-			NewInstanceComponent->OnComponentCreated();
-			NewInstanceComponent->RegisterComponent();
+				// Add to SerializedComponents array so it gets saved
+				ActorInstance->AddInstanceComponent(NewInstanceComponent);
+				NewInstanceComponent->OnComponentCreated();
+				NewInstanceComponent->RegisterComponent();
 
-			// Register any new components that may have been created during construction of the instanced component, but were not explicitly registered.
-			TInlineComponentArray<UActorComponent*> PostInstanceComponents;
-			ActorInstance->GetComponents(PostInstanceComponents);
-			for (UActorComponent* ActorComponent : PostInstanceComponents)
-			{
-				if (!ActorComponent->IsRegistered() && ActorComponent->bAutoRegister && !ActorComponent->IsPendingKill() && !PreInstanceComponents.Contains(ActorComponent))
+				// Register any new components that may have been created during construction of the instanced component, but were not explicitly registered.
+				TInlineComponentArray<UActorComponent*> PostInstanceComponents;
+				ActorInstance->GetComponents(PostInstanceComponents);
+				for (UActorComponent* ActorComponent : PostInstanceComponents)
 				{
-					ActorComponent->RegisterComponent();
+					if (!ActorComponent->IsRegistered() && ActorComponent->bAutoRegister && !ActorComponent->IsPendingKill() && !PreInstanceComponents.Contains(ActorComponent))
+					{
+						ActorComponent->RegisterComponent();
+					}
 				}
-			}
 
-			// Rerun construction scripts
-			ActorInstance->RerunConstructionScripts();
+				// Rerun construction scripts
+				ActorInstance->RerunConstructionScripts();
 
-			// If the running the construction script destroyed the new node, don't create an entry for it
-			if (!NewInstanceComponent->IsPendingKill())
+				// If the running the construction script destroyed the new node, don't create an entry for it
+				if (!NewInstanceComponent->IsPendingKill())
+				{
+					// Create a new subobject data set with this component
+					NewDataHandle = FactoryCreateSubobjectDataWithParent(NewInstanceComponent, ParentObjData->GetHandle());
+				}
+			}
+			else
 			{
-				// Create a new subobject data set with this component
-				NewDataHandle = FactoryCreateSubobjectDataWithParent(NewInstanceComponent, ParentObjData->GetHandle());
+				FailReason = LOCTEXT("AddComponentFailed_Inherited", "Cannot add components within an Inherited hierarchy");
 			}
 		}
-		else
-		{
-			FailReason =  LOCTEXT("AddComponentFailed_Inherited", "Cannot add components within an Inherited heirarchy");
-		}
 	}
 	
 	return NewDataHandle;
diff --git a/Engine/Source/Editor/SubobjectDataInterface/Public/SubobjectDataSubsystem.h b/Engine/Source/Editor/SubobjectDataInterface/Public/SubobjectDataSubsystem.h
--- a/Engine/Source/Editor/SubobjectDataInterface/Public/SubobjectDataSubsystem.h
+++ b/Engine/Source/Editor/SubobjectDataInterface/Public/SubobjectDataSubsystem.h
@@ -315,6 +315,15 @@ public:
 	UFUNCTION(BlueprintCallable, Category = "SubobjectDataSubsystem")
 	static void RenameSubobjectMemberVariable(UBlueprint* BPContext, const FSubobjectDataHandle& InHandle, const FName NewName);
 	
+	/**
+	* Walk the hierarchy up to find the AActor root for the given handle
+	* 
+	* @param StartingHandle		The handle to start walking the hierarchy at
+	*
+	* @return The handle of the root actor if one exists, FSubobjectDataHandle::InvalidHandle otherwise.
+	*/
+	static FSubobjectDataHandle GetActorRootHandle(const FSubobjectDataHandle& StartingHandle);
+
 private:
 
 	/**
