Fix typo in param name for USubobjectDataSubsystem::CreateNewBPComponent

#rnx

diff --git a/Engine/Source/Editor/SubobjectDataInterface/Public/SubobjectDataSubsystem.h b/Engine/Source/Editor/SubobjectDataInterface/Public/SubobjectDataSubsystem.h
--- a/Engine/Source/Editor/SubobjectDataInterface/Public/SubobjectDataSubsystem.h
+++ b/Engine/Source/Editor/SubobjectDataInterface/Public/SubobjectDataSubsystem.h
@@ -176,7 +176,7 @@ public:
 	 * @return The new class that was created
 	 */
 	UFUNCTION(BlueprintCallable, Category = "SubobjectDataSubsystem")
-	static UClass* CreateNewBPComponent(TSubclassOf<UActorComponent> ComponentClassconst, const FString& NewClassPath, const FString& NewClassName);
+	static UClass* CreateNewBPComponent(TSubclassOf<UActorComponent> ComponentClass, const FString& NewClassPath, const FString& NewClassName);
 	
 	/**
 	* Add a new subobject as a child to the given parent object 
