Add DisplayName metadata to CreateNewCppComponent and CreateNewBPComponent

#rnx

diff --git a/Engine/Source/Editor/SubobjectDataInterface/Public/SubobjectDataSubsystem.h b/Engine/Source/Editor/SubobjectDataInterface/Public/SubobjectDataSubsystem.h
--- a/Engine/Source/Editor/SubobjectDataInterface/Public/SubobjectDataSubsystem.h
+++ b/Engine/Source/Editor/SubobjectDataInterface/Public/SubobjectDataSubsystem.h
@@ -166,7 +166,7 @@ public:
 	 *
 	 * @return The new class that was created
 	 */
-	UFUNCTION(BlueprintCallable, Category = "SubobjectDataSubsystem")
+	UFUNCTION(BlueprintCallable, Category = "SubobjectDataSubsystem", meta = (DisplayName = "Create New C++ Component"))
 	static UClass* CreateNewCPPComponent(TSubclassOf<UActorComponent> ComponentClass, const FString& NewClassPath, const FString& NewClassName);
 
 	/**
@@ -175,7 +175,7 @@ public:
 	 *
 	 * @return The new class that was created
 	 */
-	UFUNCTION(BlueprintCallable, Category = "SubobjectDataSubsystem")
+	UFUNCTION(BlueprintCallable, Category = "SubobjectDataSubsystem", meta = (DisplayName = "Create New Blueprint Component"))
 	static UClass* CreateNewBPComponent(TSubclassOf<UActorComponent> ComponentClass, const FString& NewClassPath, const FString& NewClassName);
 	
 	/**
