Subobject Data Subsystem:
- Do not create a new blueprint component via the Subobject Data subsystem if the given class is not a blueprint type.
- Check for existing classes with that have the same name when creating a BP class to prevent the user from creating invalid packages.
- Add logging to show the user a reason for failure when the given class name is invalid.

diff --git a/Engine/Source/Editor/SubobjectDataInterface/Private/SubobjectDataSubsystem.cpp b/Engine/Source/Editor/SubobjectDataInterface/Private/SubobjectDataSubsystem.cpp
--- a/Engine/Source/Editor/SubobjectDataInterface/Private/SubobjectDataSubsystem.cpp
+++ b/Engine/Source/Editor/SubobjectDataInterface/Private/SubobjectDataSubsystem.cpp
@@ -574,25 +574,51 @@ UClass* USubobjectDataSubsystem::CreateNewCPPComponent(TSubclassOf<UActorCompone
 UClass* USubobjectDataSubsystem::CreateNewBPComponent(TSubclassOf<UActorComponent> ComponentClass, const FString& NewClassPath, const FString& NewClassName)
 {
 	UClass* NewClass = nullptr;
-	if (ComponentClass && !NewClassName.IsEmpty() && !NewClassPath.IsEmpty())
-	{
-		const FString PackagePath = NewClassPath / NewClassName;
 
-		if (UPackage* Package = CreatePackage(*PackagePath))
+	if(ComponentClass)
+	{
+		// Ensure that the class is a blueprint type
+		if(!FKismetEditorUtilities::CanCreateBlueprintOfClass(ComponentClass))
+		{
+			UE_LOG(LogSubobjectSubsystem, Warning, TEXT("Failed to Create new BP Component: '%s' is not a blueprintable class!"), *ComponentClass->GetDisplayNameText().ToString());
+			return nullptr;
+		}
+		
+		if (!NewClassName.IsEmpty() && !NewClassPath.IsEmpty())
 		{
-			// Create and init a new Blueprint
-			UBlueprint* NewBP = FKismetEditorUtilities::CreateBlueprint(ComponentClass, Package, FName(*NewClassName), BPTYPE_Normal, UBlueprint::StaticClass(), UBlueprintGeneratedClass::StaticClass());
-			if (NewBP)
+			const FString PackagePath = NewClassPath / NewClassName;
+
+			// Check for an existing object
+			if(UObject* ExistingObject = StaticFindObject(UObject::StaticClass(), ANY_PACKAGE, *PackagePath))
+			{
+				UE_LOG(LogSubobjectSubsystem, Warning, TEXT("Failed to Create new BP Component: A class with a name '%s' already exists!"), *PackagePath);
+				return nullptr;
+			}
+		
+			if (UPackage* Package = CreatePackage(*PackagePath))
 			{
-				// Notify the asset registry
-				FAssetRegistryModule::AssetCreated(NewBP);
+				// Create and init a new Blueprint			
+				if (UBlueprint* NewBP = FKismetEditorUtilities::CreateBlueprint(ComponentClass, Package, FName(*NewClassName), BPTYPE_Normal, UBlueprint::StaticClass(), UBlueprintGeneratedClass::StaticClass()))
+				{
+					// Notify the asset registry
+					FAssetRegistryModule::AssetCreated(NewBP);
 
-				// Mark the package dirty...
-				Package->MarkPackageDirty();
-				NewClass = NewBP->GeneratedClass;
+					// Mark the package dirty...
+					Package->MarkPackageDirty();
+					NewClass = NewBP->GeneratedClass;
+				}
 			}
 		}
+		else
+		{
+			UE_LOG(LogSubobjectSubsystem, Warning, TEXT("Failed to Create new BP Component: An invalid class path was given!"));
+		}
 	}
+	else
+	{
+		UE_LOG(LogSubobjectSubsystem, Warning, TEXT("Failed to Create new BP Component: ComponentClass is required!"));
+	}
+	
 	return NewClass;
 }
 
