Change how the old root component gets deleted to correctly handle the promotion of its children

#rnx

diff --git a/Engine/Source/Editor/SubobjectDataInterface/Private/SubobjectDataSubsystem.cpp b/Engine/Source/Editor/SubobjectDataInterface/Private/SubobjectDataSubsystem.cpp
--- a/Engine/Source/Editor/SubobjectDataInterface/Private/SubobjectDataSubsystem.cpp
+++ b/Engine/Source/Editor/SubobjectDataInterface/Private/SubobjectDataSubsystem.cpp
@@ -1385,7 +1385,12 @@ bool USubobjectDataSubsystem::MakeNewSceneRoot(const FSubobjectDataHandle& Conte
 			// If the thing we replaced was the default scene root, then just delete it
 			if(bWasDefaultSceneRoot)
 			{
-				DeleteSubobject(Context, OldSceneRoot, Blueprint);
+				// Destroy the old root component instance
+				UActorComponent* OldRootInstance = OldSceneRoot.GetData()->GetMutableComponentTemplate();
+				OldRootInstance->Modify();
+				OldRootInstance->DestroyComponent(!bWasDefaultSceneRoot);
+
+				// Set the new root component
 				AActor* ActorContext = Context.GetData()->GetMutableActorContext();
 				ActorContext->SetRootComponent(CastChecked<USceneComponent>(DroppedNewSceneRootData->GetMutableComponentTemplate()));
 			}
