Make DeleteSubobject functions for an instance and blueprint seperate for better context menu searches

#rnx

diff --git a/Engine/Source/Editor/SubobjectDataInterface/Private/SubobjectDataSubsystem.cpp b/Engine/Source/Editor/SubobjectDataInterface/Private/SubobjectDataSubsystem.cpp
--- a/Engine/Source/Editor/SubobjectDataInterface/Private/SubobjectDataSubsystem.cpp
+++ b/Engine/Source/Editor/SubobjectDataInterface/Private/SubobjectDataSubsystem.cpp
@@ -1157,6 +1157,16 @@ int32 USubobjectDataSubsystem::DeleteSubobjects(const FSubobjectDataHandle& Cont
 	return DeleteSubobjects(ContextHandle, SubobjectsToDelete, Dummy, BPContext);
 }
 
+int32 USubobjectDataSubsystem::K2_DeleteSubobjectFromInstance(const FSubobjectDataHandle& ContextHandle, const FSubobjectDataHandle& SubobjectToDelete)
+{
+	return DeleteSubobject(ContextHandle, SubobjectToDelete);
+}
+
+int32 USubobjectDataSubsystem::K2_DeleteSubobjectsFromInstance(const FSubobjectDataHandle& ContextHandle, const TArray<FSubobjectDataHandle>& SubobjectsToDelete)
+{
+	return DeleteSubobjects(ContextHandle, SubobjectsToDelete);
+}
+
 int32 USubobjectDataSubsystem::DeleteSubobject(const FSubobjectDataHandle& ContextHandle, const FSubobjectDataHandle& SubobjectToDelete, UBlueprint* BPContext)
 {
 	TArray<FSubobjectDataHandle> Handles { SubobjectToDelete };
diff --git a/Engine/Source/Editor/SubobjectDataInterface/Public/SubobjectDataSubsystem.h b/Engine/Source/Editor/SubobjectDataInterface/Public/SubobjectDataSubsystem.h
--- a/Engine/Source/Editor/SubobjectDataInterface/Public/SubobjectDataSubsystem.h
+++ b/Engine/Source/Editor/SubobjectDataInterface/Public/SubobjectDataSubsystem.h
@@ -193,14 +193,23 @@ public:
 	*
 	* @param ContextHandle			The owning context of the subobjects that should be removed
 	* @param SubobjectsToDelete		Array of subobject handles that should be deleted
-	*
+	* @param BPContext				The blueprint context for the given
+	* 
 	* @return 	The number of subobjects successfully deleted
 	*/
-	UFUNCTION(BlueprintCallable, Category = "SubobjectDataSubsystem")
+	UFUNCTION(BlueprintCallable, Category = "SubobjectDataSubsystem", meta = (DisplayName = "Delete Subobjects from Blueprint"))
 	int32 DeleteSubobjects(const FSubobjectDataHandle& ContextHandle, const TArray<FSubobjectDataHandle>& SubobjectsToDelete, UBlueprint* BPContext = nullptr);
 
-	UFUNCTION(BlueprintCallable, Category = "SubobjectDataSubsystem")
-    int32 DeleteSubobject(const FSubobjectDataHandle& ContextHandle, const FSubobjectDataHandle& SubobjectToDelete, UBlueprint* BPContext = nullptr);
+	/**
+	* Attempts to delete the given array of subobjects from their context 
+	*
+	* @param ContextHandle			The owning context of the subobjects that should be removed
+	* @param SubobjectsToDelete		Array of subobject handles that should be deleted
+	*
+	* @return 	The number of subobjects successfully deleted
+	*/
+	UFUNCTION(BlueprintCallable, Category = "SubobjectDataSubsystem", meta = (DisplayName = "Delete Subobjects from Instance"))
+	int32 K2_DeleteSubobjectsFromInstance(const FSubobjectDataHandle& ContextHandle, const TArray<FSubobjectDataHandle>& SubobjectsToDelete);
 	
 	/**
 	* Attempts to delete the given array of subobjects from their context 
@@ -211,6 +220,29 @@ public:
 	* @return 	The number of subobjects successfully deleted
 	*/
 	int32 DeleteSubobjects(const FSubobjectDataHandle& ContextHandle, const TArray<FSubobjectDataHandle>& SubobjectsToDelete, FSubobjectDataHandle& OutComponentToSelect, UBlueprint* BPContext = nullptr);
+
+	/**
+	* Attempts to delete the given subobject from its blueprint context 
+	*
+	* @param ContextHandle			The owning context of the subobjects that should be removed
+	* @param SubobjectToDelete		The subobject handles that should be deleted
+	* @param BPContext				The blueprint context for the given 
+	*
+	* @return 	The number of subobjects successfully deleted
+	*/
+	UFUNCTION(BlueprintCallable, Category = "SubobjectDataSubsystem", meta = (DisplayName = "Delete Subobject from Blueprint"))
+	int32 DeleteSubobject(const FSubobjectDataHandle& ContextHandle, const FSubobjectDataHandle& SubobjectToDelete, UBlueprint* BPContext = nullptr);
+	
+	/**
+	* Attempts to delete the given subobject from its context 
+	*
+	* @param ContextHandle			The owning context of the subobjects that should be removed
+	* @param SubobjectToDelete		The subobject handles that should be deleted
+	*
+	* @return 	The number of subobjects successfully deleted
+	*/
+	UFUNCTION(BlueprintCallable, Category = "SubobjectDataSubsystem", meta = (DisplayName = "Delete Subobject from Instance"))
+	int32 K2_DeleteSubobjectFromInstance(const FSubobjectDataHandle& ContextHandle, const FSubobjectDataHandle& SubobjectToDelete);
 	
 	/**
 	* Attempts to rename the given subobject to the new name.
