Remove unused parent class filter from Subobject Data Subsystem

#rnx

diff --git a/Engine/Source/Editor/SubobjectDataInterface/Private/SubobjectDataSubsystem.cpp b/Engine/Source/Editor/SubobjectDataInterface/Private/SubobjectDataSubsystem.cpp
--- a/Engine/Source/Editor/SubobjectDataInterface/Private/SubobjectDataSubsystem.cpp
+++ b/Engine/Source/Editor/SubobjectDataInterface/Private/SubobjectDataSubsystem.cpp
@@ -499,37 +499,6 @@ FSubobjectDataHandle USubobjectDataSubsystem::FindHandleForObject(const FSubobje
 	return FSubobjectDataHandle::InvalidHandle;
 }
 
-class FComponentClassParentFilter : public IClassViewerFilter
-{
-public:
-	FComponentClassParentFilter(const TSubclassOf<UActorComponent>& InComponentClass) : ComponentClass(InComponentClass) {}
-
-	virtual bool IsClassAllowed(const FClassViewerInitializationOptions& InInitOptions, const UClass* InClass, TSharedRef< FClassViewerFilterFuncs > InFilterFuncs) override
-	{
-		return InClass->IsChildOf(ComponentClass);
-	}
-
-	virtual bool IsUnloadedClassAllowed(const FClassViewerInitializationOptions& InInitOptions, const TSharedRef< const IUnloadedBlueprintData > InUnloadedClassData, TSharedRef< FClassViewerFilterFuncs > InFilterFuncs) override
-	{
-		return InUnloadedClassData->IsChildOf(ComponentClass);
-	}
-
-	TSubclassOf<UActorComponent> ComponentClass;
-};
-
-typedef FComponentClassParentFilter FNativeComponentClassParentFilter;
-
-class FBlueprintComponentClassParentFilter : public FComponentClassParentFilter
-{
-public:
-	FBlueprintComponentClassParentFilter(const TSubclassOf<UActorComponent>& InComponentClass) : FComponentClassParentFilter(InComponentClass) {}
-
-	virtual bool IsClassAllowed(const FClassViewerInitializationOptions& InInitOptions, const UClass* InClass, TSharedRef< FClassViewerFilterFuncs > InFilterFuncs) override
-	{
-		return FComponentClassParentFilter::IsClassAllowed(InInitOptions, InClass, InFilterFuncs) && FKismetEditorUtilities::CanCreateBlueprintOfClass(InClass);
-	}
-};
-
 UClass* USubobjectDataSubsystem::CreateNewCPPComponent(TSubclassOf<UActorComponent> ComponentClass, const FString& NewClassPath, const FString& NewClassName)
 {
 	UClass* NewClass = nullptr;
