Fix a crash when pasting subobjects into the blueprint editor. This will ensure that a valid attachment target is selected for the subobject, and correctly create a subobject data handle based off of the SCSnode instead of the SCSNode's component template, which was causing the crash.

#rnx

diff --git a/Engine/Source/Editor/SubobjectDataInterface/Private/SubobjectDataSubsystem.cpp b/Engine/Source/Editor/SubobjectDataInterface/Private/SubobjectDataSubsystem.cpp
--- a/Engine/Source/Editor/SubobjectDataInterface/Private/SubobjectDataSubsystem.cpp
+++ b/Engine/Source/Editor/SubobjectDataInterface/Private/SubobjectDataSubsystem.cpp
@@ -1896,21 +1896,38 @@ void USubobjectDataSubsystem::PasteSubobjects(const FSubobjectDataHandle& PasteT
 			
 			// Create a new SCS node to contain the new component and add it to the tree
 			USCS_Node* NewSCSNode = Blueprint->SimpleConstructionScript->CreateNodeAndRenameComponent(NewActorComponent);
+			NewSCSNode->Modify();
 			NewActorComponent = NewSCSNode ? ToRawPtr(NewSCSNode->ComponentTemplate) : nullptr;
 
-			FSubobjectDataHandle TargetParentHandle = FindParentForNewSubobject(NewActorComponent, PasteToContext);
+			// The default target to attach the new pasted component to is the Scene Root for the context
+			FSubobjectDataHandle TargetParentHandle = FindSceneRootForSubobject(FindParentForNewSubobject(NewActorComponent, PasteToContext));
+
+			// If there were any handles that we should target the attachment for instead, use that
+			for (FSubobjectDataHandle SelectedNode : NewParentHandles)
+			{
+				const FSubobjectData* SelectedData = SelectedNode.GetSharedDataPtr().Get();
+
+				// Only scene components can be attached to when pasting
+				if (SelectedData && SelectedData->IsSceneComponent())
+				{
+					TargetParentHandle = SelectedData->GetHandle();
+					break;
+				}
+			}
+
 			TSharedPtr<FInheritedSubobjectData> TargetData = StaticCastSharedPtr<FInheritedSubobjectData>(TargetParentHandle.GetSharedDataPtr());
 			
 			// Create a new subobject data set with this component. Use the SCS node here and the subobject data
 			// will correctly associate the component template
 			FSubobjectDataHandle NewDataHandle = FactoryCreateSubobjectDataWithParent(
-				NewActorComponent,
+				NewSCSNode,
 				TargetParentHandle,
 				TargetData ? TargetData->bIsInheritedSCS : false
 			);
 
 			AttachSubobject(TargetParentHandle, NewDataHandle);
 
+			OutPastedHandles.Add(NewDataHandle);
 			// Map the new subobject's data handle to it's instance name
 			NewNodeMap.Add(NewObjectPair.Key, NewDataHandle);
 		}
