Fix a crash when deleteing instances of subobjects if the component template is null

#rn trivial
#rnx

diff --git a/Engine/Source/Editor/SubobjectDataInterface/Private/SubobjectDataSubsystem.cpp b/Engine/Source/Editor/SubobjectDataInterface/Private/SubobjectDataSubsystem.cpp
--- a/Engine/Source/Editor/SubobjectDataInterface/Private/SubobjectDataSubsystem.cpp
+++ b/Engine/Source/Editor/SubobjectDataInterface/Private/SubobjectDataSubsystem.cpp
@@ -1107,7 +1107,10 @@ int32 USubobjectDataSubsystem::DeleteSubobjects(const FSubobjectDataHandle& Cont
 			{
 				if(const FSubobjectData* Data = Handle.GetData())
 				{
-					ComponentsToDelete.Add(Data->GetMutableComponentTemplate());		
+					if(UActorComponent* Template = Data->GetMutableComponentTemplate())
+					{
+						ComponentsToDelete.Add(Template);
+					}					
 				}		
 			}
 		}
