Add clearer logging when creating a new BP component

#rnx

diff --git a/Engine/Source/Editor/SubobjectDataInterface/Private/SubobjectDataSubsystem.cpp b/Engine/Source/Editor/SubobjectDataInterface/Private/SubobjectDataSubsystem.cpp
--- a/Engine/Source/Editor/SubobjectDataInterface/Private/SubobjectDataSubsystem.cpp
+++ b/Engine/Source/Editor/SubobjectDataInterface/Private/SubobjectDataSubsystem.cpp
@@ -547,13 +547,20 @@ UClass* USubobjectDataSubsystem::CreateNewBPComponent(TSubclassOf<UActorComponen
 	if(ComponentClass)
 	{
 		// Ensure that the class is a blueprint type
+		// Add logging to make it easier to debug what is going on if you are calling this function from a Blueprint
 		if(!FKismetEditorUtilities::CanCreateBlueprintOfClass(ComponentClass))
 		{
 			UE_LOG(LogSubobjectSubsystem, Warning, TEXT("Failed to Create new BP Component: '%s' is not a blueprintable class!"), *ComponentClass->GetDisplayNameText().ToString());
-			return nullptr;
 		}
-		
-		if (!NewClassName.IsEmpty() && !NewClassPath.IsEmpty())
+		else if (NewClassName.IsEmpty())
+		{
+			UE_LOG(LogSubobjectSubsystem, Warning, TEXT("Failed to Create new BP Component: An NewClassName was given!"));
+		}
+		else if (NewClassPath.IsEmpty())
+		{
+			UE_LOG(LogSubobjectSubsystem, Warning, TEXT("Failed to Create new BP Component: An invalid class path was given!"));
+		}
+		else
 		{
 			const FString PackagePath = NewClassPath / NewClassName;
 
@@ -578,10 +585,6 @@ UClass* USubobjectDataSubsystem::CreateNewBPComponent(TSubclassOf<UActorComponen
 				}
 			}
 		}
-		else
-		{
-			UE_LOG(LogSubobjectSubsystem, Warning, TEXT("Failed to Create new BP Component: An invalid class path was given!"));
-		}
 	}
 	else
 	{
