Fix SA warning with NewSCSNode in PasteSubobjects

#rnx

diff --git a/Engine/Source/Editor/SubobjectDataInterface/Private/SubobjectDataSubsystem.cpp b/Engine/Source/Editor/SubobjectDataInterface/Private/SubobjectDataSubsystem.cpp
--- a/Engine/Source/Editor/SubobjectDataInterface/Private/SubobjectDataSubsystem.cpp
+++ b/Engine/Source/Editor/SubobjectDataInterface/Private/SubobjectDataSubsystem.cpp
@@ -1902,8 +1902,15 @@ void USubobjectDataSubsystem::PasteSubobjects(const FSubobjectDataHandle& PasteT
 			
 			// Create a new SCS node to contain the new component and add it to the tree
 			USCS_Node* NewSCSNode = Blueprint->SimpleConstructionScript->CreateNodeAndRenameComponent(NewActorComponent);
-			NewSCSNode->Modify();
-			NewActorComponent = NewSCSNode ? ToRawPtr(NewSCSNode->ComponentTemplate) : nullptr;
+			if(NewSCSNode)
+			{
+				NewSCSNode->Modify();
+				NewActorComponent = ToRawPtr(NewSCSNode->ComponentTemplate);
+			}
+			else
+			{
+				NewActorComponent = nullptr;
+			}
 
 			// The default target to attach the new pasted component to is the Scene Root for the context
 			FSubobjectDataHandle TargetParentHandle = FindSceneRootForSubobject(FindParentForNewSubobject(NewActorComponent, PasteToContext));
