Fix duplication of an inherited blueprint variables in the subobject instance editor by using the correct ComponentInstanceInActor instead of the SCS node

#rnx

diff --git a/Engine/Source/Editor/SubobjectDataInterface/Private/SubobjectDataSubsystem.cpp b/Engine/Source/Editor/SubobjectDataInterface/Private/SubobjectDataSubsystem.cpp
--- a/Engine/Source/Editor/SubobjectDataInterface/Private/SubobjectDataSubsystem.cpp
+++ b/Engine/Source/Editor/SubobjectDataInterface/Private/SubobjectDataSubsystem.cpp
@@ -2014,7 +2014,7 @@ void USubobjectDataSubsystem::DuplicateSubobjects(const FSubobjectDataHandle& Co
 	
 	TMap<FSubobjectData*, FSubobjectData*> DuplicateSceneComponentMap;
 
-	// For each Subobject to dup, add it as a new subobejct to the context
+	// For each Subobject to dup, add it as a new subobject to the context
 	for(const FSubobjectDataHandle& OriginalHandle : SubobjectsToDup)
 	{
 		if(!OriginalHandle.IsValid())
@@ -2024,13 +2024,26 @@ void USubobjectDataSubsystem::DuplicateSubobjects(const FSubobjectDataHandle& Co
 		}
 
 		FSubobjectData* OriginalData = OriginalHandle.GetSharedDataPtr().Get();
+		NewSubobjectParams.ParentHandle = OriginalData->GetHandle();
+
 		if(UActorComponent* ComponentTemplate = OriginalData->GetMutableComponentTemplate())
 		{
 			USCS_Node* SCSNode = OriginalData->GetSCSNode();
 			check(SCSNode == nullptr || SCSNode->ComponentTemplate == ComponentTemplate);
 
 			NewSubobjectParams.NewClass = ComponentTemplate->GetClass();
-			NewSubobjectParams.AssetOverride = SCSNode ? (UObject*)SCSNode : ComponentTemplate;
+			if (BpContext)
+			{
+				NewSubobjectParams.AssetOverride = SCSNode ? (UObject*)SCSNode : ComponentTemplate;
+			}
+			else
+			{
+				FSubobjectData* ContextData = Context.GetData();
+				AActor* ActorContext = ContextData->GetMutableActorContext();
+				UActorComponent* InstanceComponent = OriginalData->FindMutableComponentInstanceInActor(ActorContext);
+				
+				NewSubobjectParams.AssetOverride = InstanceComponent;
+			}
 			
 			FSubobjectDataHandle ClonedSubobject = AddNewSubobject(NewSubobjectParams, FailedAddReason);
 			FSubobjectData* ClonedData = ClonedSubobject.GetSharedDataPtr().Get();
