Remove an unnecessary check from the deletion of subobjects. This check was originally there to confirm that the SCS node being deleted is from the current subobject tree, but is not a valid check anymore when you can call this function from Blueprints as well.

#rnx

diff --git a/Engine/Source/Editor/SubobjectDataInterface/Private/SubobjectDataSubsystem.cpp b/Engine/Source/Editor/SubobjectDataInterface/Private/SubobjectDataSubsystem.cpp
--- a/Engine/Source/Editor/SubobjectDataInterface/Private/SubobjectDataSubsystem.cpp
+++ b/Engine/Source/Editor/SubobjectDataInterface/Private/SubobjectDataSubsystem.cpp
@@ -1007,8 +1007,7 @@ int32 USubobjectDataSubsystem::DeleteSubobjects(const FSubobjectDataHandle& Cont
 					
 					if (SCS_Node)
 					{
-						USimpleConstructionScript* SCS = SCS_Node->GetSCS();
-						check(SCS != nullptr && BPContext == SCS->GetBlueprint());
+						USimpleConstructionScript* SCS = SCS_Node->GetSCS();						
 						BPContext->Modify();
 						SaveSCSCurrentState(SCS);
 					
