Added logging when a blueprint component cannot be deleted when called from blueprints

#rnx

diff --git a/Engine/Source/Editor/SubobjectDataInterface/Private/SubobjectDataSubsystem.cpp b/Engine/Source/Editor/SubobjectDataInterface/Private/SubobjectDataSubsystem.cpp
--- a/Engine/Source/Editor/SubobjectDataInterface/Private/SubobjectDataSubsystem.cpp
+++ b/Engine/Source/Editor/SubobjectDataInterface/Private/SubobjectDataSubsystem.cpp
@@ -1092,6 +1092,10 @@ int32 USubobjectDataSubsystem::DeleteSubobjects(const FSubobjectDataHandle& Cont
 								}
 							}
 						}
+						else
+						{
+							UE_LOG(LogSubobjectSubsystem, Warning, TEXT("Cannot remove subobject '%s' because it is the default scene root!"), *Data->GetDisplayString());
+						}
 					}
 				}
 			}
