Check if a subobject is allowed to be deleted before allowing the subsystem to delete it. This solves the problem where a BP script may not be doing the same validation when calling DeleteSubobjects that the Subobject Editor does.

Add logging when a subobject can't be deleted

#rnx

diff --git a/Engine/Source/Editor/SubobjectDataInterface/Private/SubobjectDataSubsystem.cpp b/Engine/Source/Editor/SubobjectDataInterface/Private/SubobjectDataSubsystem.cpp
--- a/Engine/Source/Editor/SubobjectDataInterface/Private/SubobjectDataSubsystem.cpp
+++ b/Engine/Source/Editor/SubobjectDataInterface/Private/SubobjectDataSubsystem.cpp
@@ -1003,6 +1003,12 @@ int32 USubobjectDataSubsystem::DeleteSubobjects(const FSubobjectDataHandle& Cont
 			{
 				if (const FSubobjectData* Data = Handle.GetData())
 				{
+					if (!Data->CanDelete())
+					{
+						UE_LOG(LogSubobjectSubsystem, Warning, TEXT("Cannot delete subobject '%s' from '%s'!"), *Data->GetDisplayString(), *BPContext->GetFullName());
+						continue;
+					}
+
 					USCS_Node* SCS_Node = Data->GetSCSNode();
 					
 					if (SCS_Node)
@@ -1113,6 +1119,12 @@ int32 USubobjectDataSubsystem::DeleteSubobjects(const FSubobjectDataHandle& Cont
 			{
 				if(const FSubobjectData* Data = Handle.GetData())
 				{
+					if (!Data->CanDelete())
+					{
+						UE_LOG(LogSubobjectSubsystem, Warning, TEXT("Cannot delete subobject '%s' from '%s'!"), *Data->GetDisplayString(), *ContextObj->GetFullName());
+						continue;
+					}
+
 					if(UActorComponent* Template = Data->GetMutableComponentTemplate())
 					{
 						ComponentsToDelete.Add(Template);
@@ -1120,10 +1132,13 @@ int32 USubobjectDataSubsystem::DeleteSubobjects(const FSubobjectDataHandle& Cont
 				}		
 			}
 		}
-		UActorComponent* ActorComponentToSelect = nullptr;
-		NumDeletedSubobjects = FComponentEditorUtils::DeleteComponents(ComponentsToDelete, ActorComponentToSelect);
 
-		// Find the handle that matches with this actor component given our current context
+		// Actually delete the components if we have any that can be
+		if (!ComponentsToDelete.IsEmpty())
+		{
+			UActorComponent* ActorComponentToSelect = nullptr;
+			NumDeletedSubobjects = FComponentEditorUtils::DeleteComponents(ComponentsToDelete, ActorComponentToSelect);
+		}
 	}
 
 	return NumDeletedSubobjects;
