Fix typo in comment
#rnx

diff --git a/Engine/Source/Editor/SubobjectDataInterface/Public/SubobjectDataSubsystem.h b/Engine/Source/Editor/SubobjectDataInterface/Public/SubobjectDataSubsystem.h
--- a/Engine/Source/Editor/SubobjectDataInterface/Public/SubobjectDataSubsystem.h
+++ b/Engine/Source/Editor/SubobjectDataInterface/Public/SubobjectDataSubsystem.h
@@ -286,7 +286,7 @@ public:
 	bool DetachSubobject(const FSubobjectDataHandle& OwnerHandle, const FSubobjectDataHandle& ChildToRemove);
 
 	/**
-	 * Addd the given subobject to a new owner. This will remove the subobject from its previous 
+	 * Add the given subobject to a new owner. This will remove the subobject from its previous 
 	 * owner if necessary.
 	 *
 	 * @param OwnerHandle		The new owner to attach to
