Attach a dupliucated component to the source's parent instead of the source directly

#rnx

diff --git a/Engine/Source/Editor/SubobjectDataInterface/Private/SubobjectDataSubsystem.cpp b/Engine/Source/Editor/SubobjectDataInterface/Private/SubobjectDataSubsystem.cpp
--- a/Engine/Source/Editor/SubobjectDataInterface/Private/SubobjectDataSubsystem.cpp
+++ b/Engine/Source/Editor/SubobjectDataInterface/Private/SubobjectDataSubsystem.cpp
@@ -2105,14 +2105,7 @@ void USubobjectDataSubsystem::DuplicateSubobjects(const FSubobjectDataHandle& Co
 			FSubobjectDataHandle ParentHandle = OriginalData->GetParentHandle();
 			if(ParentHandle.IsValid())
 			{
-				if (FSubobjectData** ParentDuplicateComponent = DuplicateSceneComponentMap.Find(ParentHandle.GetData()))
-				{
-					//FSCSEditorTreeNodePtrType DuplicateParentNodePtr = FindTreeNode(*ParentDuplicateComponent);
-					//if (DuplicateParentNodePtr.IsValid())
-					{
-						//ParentNodePtr = DuplicateParentNodePtr;
-					}
-				}
+				AttachSubobject(ParentHandle, NewData->GetHandle());
 			}
 		}
 	}
