Notify the level editor that a subobject has been changed when adding or deleting a subobject. This ensures that if there is a selected actor instance in the level that it will have it's subobjects properly updated

#rnx

diff --git a/Engine/Source/Editor/SubobjectDataInterface/Private/SubobjectDataSubsystem.cpp b/Engine/Source/Editor/SubobjectDataInterface/Private/SubobjectDataSubsystem.cpp
--- a/Engine/Source/Editor/SubobjectDataInterface/Private/SubobjectDataSubsystem.cpp
+++ b/Engine/Source/Editor/SubobjectDataInterface/Private/SubobjectDataSubsystem.cpp
@@ -2,6 +2,11 @@
 
 #include "SubobjectDataSubsystem.h"
 #include "UnrealEngine.h"				// GEngine for subsystems
+
+#if WITH_EDITOR
+#include "LevelEditor.h"				// LevelEditor.BroadcastComponentsEdited
+#endif	// WITH_EDITOR
+
 #include "Editor/EditorEngine.h"		// FActorLabelUtilities
 #include "ComponentAssetBroker.h"		// FComponentAssetBrokerage
 #include "ChildActorSubobjectData.h"
@@ -36,6 +41,15 @@
 
 DEFINE_LOG_CATEGORY_STATIC(LogSubobjectSubsystem, Log, All);
 
+/** Notify the Level Editor that there have been subobject changes to an instance and it needs to be refreshed */
+static void BroadcastInstanceChanges()
+{
+#if WITH_EDITOR
+	FLevelEditorModule& LevelEditor = FModuleManager::LoadModuleChecked<FLevelEditorModule>("LevelEditor");
+	LevelEditor.BroadcastComponentsEdited();
+#endif	// WITH_EDITOR
+}
+
 //////////////////////////////////////////////
 // USubobjectDataSubsystem
 
@@ -970,6 +984,8 @@ FSubobjectDataHandle USubobjectDataSubsystem::AddNewSubobject(const FAddNewSubob
 					// Create a new subobject data set with this component
 					NewDataHandle = FactoryCreateSubobjectDataWithParent(NewInstanceComponent, ParentObjData->GetHandle());
 				}
+				
+				BroadcastInstanceChanges();
 			}
 			else
 			{
@@ -1138,6 +1154,7 @@ int32 USubobjectDataSubsystem::DeleteSubobjects(const FSubobjectDataHandle& Cont
 		{
 			UActorComponent* ActorComponentToSelect = nullptr;
 			NumDeletedSubobjects = FComponentEditorUtils::DeleteComponents(ComponentsToDelete, ActorComponentToSelect);
+			BroadcastInstanceChanges();
 		}
 	}
 
@@ -1555,6 +1572,8 @@ bool USubobjectDataSubsystem::ReparentSubobjects(const FReparentSubobjectParams&
 		{
 			ActorInstance->RerunConstructionScripts();
 		}
+
+		BroadcastInstanceChanges();
 	}
 	
 	return true;
