Fixed subobject editor not honoring the requested post-delete selection

diff --git a/Engine/Source/Editor/SubobjectDataInterface/Private/SubobjectDataSubsystem.cpp b/Engine/Source/Editor/SubobjectDataInterface/Private/SubobjectDataSubsystem.cpp
--- a/Engine/Source/Editor/SubobjectDataInterface/Private/SubobjectDataSubsystem.cpp
+++ b/Engine/Source/Editor/SubobjectDataInterface/Private/SubobjectDataSubsystem.cpp
@@ -1154,6 +1154,10 @@ int32 USubobjectDataSubsystem::DeleteSubobjects(const FSubobjectDataHandle& Cont
 		{
 			UActorComponent* ActorComponentToSelect = nullptr;
 			NumDeletedSubobjects = FComponentEditorUtils::DeleteComponents(ComponentsToDelete, ActorComponentToSelect);
+			if (ActorComponentToSelect)
+			{
+				OutComponentToSelect = FindHandleForObject(ContextHandle, ActorComponentToSelect);
+			}
 			BroadcastInstanceChanges();
 		}
 	}
