Correctly set the module path when creating a C++ class from the subobject data subsystem. This will correctly set the module path so that the error messages are correct from GameProjectUtils::AddCodeToProject. This will make the module selection the same as how the SNewClassDialog handles it for a more consistent and correct experience.

#rnx

diff --git a/Engine/Source/Editor/SubobjectDataInterface/Private/SubobjectDataSubsystem.cpp b/Engine/Source/Editor/SubobjectDataInterface/Private/SubobjectDataSubsystem.cpp
--- a/Engine/Source/Editor/SubobjectDataInterface/Private/SubobjectDataSubsystem.cpp
+++ b/Engine/Source/Editor/SubobjectDataInterface/Private/SubobjectDataSubsystem.cpp
@@ -513,6 +513,69 @@ FSubobjectDataHandle USubobjectDataSubsystem::FindHandleForObject(const FSubobje
 	return FSubobjectDataHandle::InvalidHandle;
 }
 
+static TSharedPtr<FModuleContextInfo> GetModuleForCppCreation()
+{
+	TSharedPtr<FModuleContextInfo> OutModuleInfo;
+
+	TArray<TSharedPtr<FModuleContextInfo>> AvailableModules;
+
+	{
+		TArray<FModuleContextInfo> CurrentModules = GameProjectUtils::GetCurrentProjectModules();
+		check(CurrentModules.Num()); // this should never happen since GetCurrentProjectModules is supposed to add a dummy runtime module if the project currently has no modules
+
+		TArray<FModuleContextInfo> CurrentPluginModules = GameProjectUtils::GetCurrentProjectPluginModules();
+
+		CurrentModules.Append(MoveTemp(CurrentPluginModules));
+
+		AvailableModules.Reserve(CurrentModules.Num());
+		for (const FModuleContextInfo& ModuleInfo : CurrentModules)
+		{
+			AvailableModules.Emplace(MakeShareable(new FModuleContextInfo(ModuleInfo)));
+		}
+	}
+
+	const FString ProjectName = FApp::GetProjectName();
+
+	// Find initially selected module based on simple fallback in this order..
+	// Previously selected module, main project module, a  runtime module
+	TSharedPtr<FModuleContextInfo> ProjectModule;
+	TSharedPtr<FModuleContextInfo> RuntimeModule;
+
+	for (const auto& AvailableModule : AvailableModules)
+	{
+		if (AvailableModule->ModuleName == ProjectName)
+		{
+			ProjectModule = AvailableModule;
+		}
+
+		if (AvailableModule->ModuleType == EHostType::Runtime)
+		{
+			RuntimeModule = AvailableModule;
+		}
+	}
+
+	if (!OutModuleInfo.IsValid())
+	{
+		if (ProjectModule.IsValid())
+		{
+			// use the project module we found
+			OutModuleInfo = ProjectModule;
+		}
+		else if (RuntimeModule.IsValid())
+		{
+			// use the first runtime module we found
+			OutModuleInfo = RuntimeModule;
+		}
+		else
+		{
+			// default to just the first module
+			OutModuleInfo = AvailableModules[0];
+		}
+	}
+
+	return OutModuleInfo;
+}
+
 UClass* USubobjectDataSubsystem::CreateNewCPPComponent(TSubclassOf<UActorComponent> ComponentClass, const FString& NewClassPath, const FString& NewClassName)
 {
 	UClass* NewClass = nullptr;
@@ -523,7 +586,7 @@ UClass* USubobjectDataSubsystem::CreateNewCPPComponent(TSubclassOf<UActorCompone
 		FString CppFilePath;
 		FText FailReason;
 
-		TSharedPtr<FModuleContextInfo> SelectedModuleInfo = MakeShareable(new FModuleContextInfo());
+		TSharedPtr<FModuleContextInfo> SelectedModuleInfo = GetModuleForCppCreation();
 		FNewClassInfo NewClassInfo(ComponentClass);
 
 		// Attempt to add new source files to the project
