Use CopyPropertiesForUnrelatedObjects instead of a basic ObjectWriter/Reader combo when duplicating components to properly handle subobjects.

#rnx

diff --git a/Engine/Source/Editor/SubobjectDataInterface/Private/SubobjectDataSubsystem.cpp b/Engine/Source/Editor/SubobjectDataInterface/Private/SubobjectDataSubsystem.cpp
--- a/Engine/Source/Editor/SubobjectDataInterface/Private/SubobjectDataSubsystem.cpp
+++ b/Engine/Source/Editor/SubobjectDataInterface/Private/SubobjectDataSubsystem.cpp
@@ -910,10 +910,7 @@ FSubobjectDataHandle USubobjectDataSubsystem::AddNewSubobject(const FAddNewSubob
 		
 		if (ComponentTemplate)
 		{
-			//Serialize object properties using write/read operations.
-			TArray<uint8> SavedProperties;
-			FObjectWriter Writer(ComponentTemplate, SavedProperties);
-			FObjectReader(NewComponent, SavedProperties);
+			UEngine::CopyPropertiesForUnrelatedObjects(ComponentTemplate, NewComponent);
 			NewComponent->UpdateComponentToWorld();
 		}
 
@@ -1367,9 +1364,7 @@ bool USubobjectDataSubsystem::MakeNewSceneRoot(const FSubobjectDataHandle& Conte
 			check(ClonedComponent);
 
 			// Serialize object properties using write/read operations.
-			TArray<uint8> SavedProperties;
-			FObjectWriter Writer(ComponentTemplate, SavedProperties);
-			FObjectReader(ClonedComponent, SavedProperties);
+			UEngine::CopyPropertiesForUnrelatedObjects(ComponentTemplate, ClonedComponent);
 
 			DroppedNewSceneRootData = ClonedHandle.GetData();
 			check(DroppedNewSceneRootData->IsValid());
@@ -1537,10 +1532,7 @@ bool USubobjectDataSubsystem::ReparentSubobjects(const FReparentSubobjectParams&
 				FSubobjectData* ClonedData = ClonedSubobject.GetSharedDataPtr().Get();
 				UActorComponent* ClonedComponent = ClonedData->GetMutableComponentTemplate();
 	
-				// Serialize object properties using write/read operations.
-				TArray<uint8> SavedProperties;
-				FObjectWriter Writer(ComponentTemplate, SavedProperties);
-				FObjectReader(ClonedComponent, SavedProperties);			
+				UEngine::CopyPropertiesForUnrelatedObjects(ComponentTemplate, ClonedComponent);		
 			}
 			else
 			{
