Instead of checking if something is a child actor component when attempting to maniuplate it, check if its owner is a CAC instead. If it's owner is a CAC, then it cannot be removed because it is determined by the owning CAC. Before, the code was incorrectly checking if something was just a plain old child actor component, resulting in being unable to delete, reparent, or attach to a child actor component.

#rnx

diff --git a/Engine/Source/Editor/SubobjectDataInterface/Private/ChildActorSubobjectData.cpp b/Engine/Source/Editor/SubobjectDataInterface/Private/ChildActorSubobjectData.cpp
--- a/Engine/Source/Editor/SubobjectDataInterface/Private/ChildActorSubobjectData.cpp
+++ b/Engine/Source/Editor/SubobjectDataInterface/Private/ChildActorSubobjectData.cpp
@@ -29,21 +29,4 @@ FText FChildActorSubobjectData::GetActorDisplayText() const
 bool FChildActorSubobjectData::IsChildActor() const
 {
 	return true;
-}
-
-bool FChildActorSubobjectData::CanDelete() const
-{
-	return false;
-}
-
-bool FChildActorSubobjectData::CanReparent() const
-{
-	// Can't reparent nodes within a child actor subtree
-	return false;
-}
-
-bool FChildActorSubobjectData::CanDuplicate() const
-{
-	// Duplication of Child actor components is not allowed.
-	return false;
-}	
\ No newline at end of file
+}
\ No newline at end of file
diff --git a/Engine/Source/Editor/SubobjectDataInterface/Private/InheritedSubobjectData.cpp b/Engine/Source/Editor/SubobjectDataInterface/Private/InheritedSubobjectData.cpp
--- a/Engine/Source/Editor/SubobjectDataInterface/Private/InheritedSubobjectData.cpp
+++ b/Engine/Source/Editor/SubobjectDataInterface/Private/InheritedSubobjectData.cpp
@@ -44,7 +44,7 @@ bool FInheritedSubobjectData::CanEdit() const
 
 bool FInheritedSubobjectData::CanDelete() const
 {
-	if(IsNativeComponent() || IsInheritedSCSNode() || IsDefaultSceneRoot() || IsInstancedInheritedComponent())
+	if(IsInheritedComponent() || IsDefaultSceneRoot() || (GetSCSNode() != nullptr && IsInstancedInheritedComponent()) || IsChildActorSubtreeObject())
 	{
 		return false;
 	}
@@ -52,15 +52,6 @@ bool FInheritedSubobjectData::CanDelete() const
 	return true;
 }
 
-bool FInheritedSubobjectData::IsInheritedComponent() const
-{
-	if(GetComponentTemplate() != nullptr)
-	{
-		return true;
-	}
-	return false;
-}
-
 bool FInheritedSubobjectData::IsInheritedSCSNode() const
 {
 	return bIsInheritedSCS;
diff --git a/Engine/Source/Editor/SubobjectDataInterface/Private/SubobjectData.cpp b/Engine/Source/Editor/SubobjectDataInterface/Private/SubobjectData.cpp
--- a/Engine/Source/Editor/SubobjectDataInterface/Private/SubobjectData.cpp
+++ b/Engine/Source/Editor/SubobjectDataInterface/Private/SubobjectData.cpp
@@ -108,12 +108,12 @@ bool FSubobjectData::CanRename() const
 
 	// You can rename within the BP editor, but not if if this subobject
 	// is on an instance in the level
-	if(GetSCSNode() != nullptr)
+	if(GetSCSNode() != nullptr && !IsInheritedSCSNode())
 	{
 		return !IsInstancedInheritedComponent();
 	}
 	
-	return !IsInheritedComponent() && !IsDefaultSceneRoot() && !IsChildActor();
+	return !IsInheritedComponent() && !IsDefaultSceneRoot() && !IsChildActorSubtreeObject();
 }
 
 const UObject* FSubobjectData::GetObjectForBlueprint(UBlueprint* Blueprint) const
@@ -920,6 +920,16 @@ bool FSubobjectData::IsChildActor() const
 	return false;
 }
 
+bool FSubobjectData::IsChildActorSubtreeObject() const
+{
+	const FSubobjectDataHandle& RootActor = GetRootSubobject();
+	if (const FSubobjectData* RootData = RootActor.GetData())
+	{
+		return RootData->IsChildActor();
+	}
+	return false;
+}
+
 bool FSubobjectData::IsRootActor() const
 {
 	// This is the root actor if it points to an AActor and has no parent
diff --git a/Engine/Source/Editor/SubobjectDataInterface/Private/SubobjectDataSubsystem.cpp b/Engine/Source/Editor/SubobjectDataInterface/Private/SubobjectDataSubsystem.cpp
--- a/Engine/Source/Editor/SubobjectDataInterface/Private/SubobjectDataSubsystem.cpp
+++ b/Engine/Source/Editor/SubobjectDataInterface/Private/SubobjectDataSubsystem.cpp
@@ -747,7 +747,7 @@ FSubobjectDataHandle USubobjectDataSubsystem::FindParentForNewSubobject(const UO
 	FSubobjectData* TargetParentData = TargetParentHandle.GetData();
 
 	// If the current selection belongs to a child actor template, move the target to its outer component node.
-	while(TargetParentHandle.IsValid() && TargetParentData->IsChildActor())
+	while(TargetParentHandle.IsValid() && TargetParentData->IsChildActorSubtreeObject())
 	{
 		TargetParentHandle = TargetParentData->GetParentHandle();
 		TargetParentData = TargetParentHandle.GetData();
diff --git a/Engine/Source/Editor/SubobjectDataInterface/Public/ChildActorSubobjectData.h b/Engine/Source/Editor/SubobjectDataInterface/Public/ChildActorSubobjectData.h
--- a/Engine/Source/Editor/SubobjectDataInterface/Public/ChildActorSubobjectData.h
+++ b/Engine/Source/Editor/SubobjectDataInterface/Public/ChildActorSubobjectData.h
@@ -15,9 +15,6 @@ struct SUBOBJECTDATAINTERFACE_API FChildActorSubobjectData final : public FInher
 	virtual FText GetDisplayName() const override;
 	virtual FText GetActorDisplayText() const override;
 	virtual bool IsChildActor() const override;
-	virtual bool CanDelete() const override;
-	virtual bool CanReparent() const override;
-	virtual bool CanDuplicate() const override;
 	// End FSubobjectData
 
 	inline const UChildActorComponent* GetChildActorComponent(bool bEvenIfPendingKill = false) const
diff --git a/Engine/Source/Editor/SubobjectDataInterface/Public/InheritedSubobjectData.h b/Engine/Source/Editor/SubobjectDataInterface/Public/InheritedSubobjectData.h
--- a/Engine/Source/Editor/SubobjectDataInterface/Public/InheritedSubobjectData.h
+++ b/Engine/Source/Editor/SubobjectDataInterface/Public/InheritedSubobjectData.h
@@ -18,7 +18,6 @@ struct SUBOBJECTDATAINTERFACE_API FInheritedSubobjectData : public FSubobjectDat
 	virtual bool CanEdit() const override;
 	virtual bool CanDelete() const override;
 	virtual bool IsInheritedSCSNode() const override;
-	virtual bool IsInheritedComponent() const override;
 	// End FSubobjectData
 
 protected:
diff --git a/Engine/Source/Editor/SubobjectDataInterface/Public/SubobjectData.h b/Engine/Source/Editor/SubobjectDataInterface/Public/SubobjectData.h
--- a/Engine/Source/Editor/SubobjectDataInterface/Public/SubobjectData.h
+++ b/Engine/Source/Editor/SubobjectDataInterface/Public/SubobjectData.h
@@ -129,6 +129,9 @@ public:
 	/* Returns true if this subobject is a component. */
 	virtual bool IsComponent() const;
 	virtual bool IsChildActor() const;
+
+	/** Returns true if this subobject is a part of a child actor's hierarchy and not the root actor.  */
+	virtual bool IsChildActorSubtreeObject() const;
 	virtual bool IsRootActor() const;
 	virtual bool IsActor() const;
 	/** True if this subobject is an instance of an inherited component (its owner is an instanced actor) */
diff --git a/Engine/Source/Editor/SubobjectEditor/Private/SSubobjectEditor.cpp b/Engine/Source/Editor/SubobjectEditor/Private/SSubobjectEditor.cpp
--- a/Engine/Source/Editor/SubobjectEditor/Private/SSubobjectEditor.cpp
+++ b/Engine/Source/Editor/SubobjectEditor/Private/SSubobjectEditor.cpp
@@ -221,7 +221,7 @@ FSubobjectData* FSubobjectEditorTreeNode::GetDataSource() const
 bool FSubobjectEditorTreeNode::IsChildSubtreeNode() const
 {
 	const FSubobjectData* Data = GetDataSource();
-	return Data ? Data->IsChildActor() : false;
+	return Data ? Data->IsChildActorSubtreeObject() : false;
 }
 
 bool FSubobjectEditorTreeNode::IsAttachedTo(FSubobjectEditorTreeNodePtrType InNodePtr) const
@@ -2685,15 +2685,6 @@ bool SSubobjectEditor::CanDeleteNodes() const
 		{
 			return false;
 		}
-		else
-		{
-			// Don't allow nodes that belong to a child actor template to be deleted
-			const bool bIsChildActorSubtreeNode = (SelectedNode->GetDataSource()->IsChildActor());
-			if (bIsChildActorSubtreeNode)
-			{
-				return false;
-			}
-		}
 	}
 	
 	return true;
