Correctly check if a component is inherited before allowing it to be reparented

#rnx

diff --git a/Engine/Source/Editor/SubobjectDataInterface/Private/SubobjectData.cpp b/Engine/Source/Editor/SubobjectDataInterface/Private/SubobjectData.cpp
--- a/Engine/Source/Editor/SubobjectDataInterface/Private/SubobjectData.cpp
+++ b/Engine/Source/Editor/SubobjectDataInterface/Private/SubobjectData.cpp
@@ -78,12 +78,12 @@ bool FSubobjectData::CanCopy() const
 }
 
 bool FSubobjectData::CanReparent() const
-{	
+{
 	if(IsComponent())
 	{
 		if(GetSCSNode() != nullptr)
 		{
-			return !IsInstancedInheritedComponent();
+			return !IsInstancedInheritedComponent() && !IsInheritedComponent();
 		}
 		
 		return !IsInheritedComponent() && !IsDefaultSceneRoot() && IsSceneComponent();
@@ -830,8 +830,6 @@ bool FSubobjectData::IsNativeComponent() const
 
 bool FSubobjectData::IsInheritedComponent() const
 {
-	// If this component is not created via an instance, then it's an inherited component
-
 	// This covers a component that is added via blueprints
 	if(GetSCSNode() != nullptr)
 	{
