If the attachment point is a BP inherited subobject then find the correct archetype attachment point when adding subobjects to an instance

#rnx

diff --git a/Engine/Source/Editor/SubobjectDataInterface/Private/SubobjectDataSubsystem.cpp b/Engine/Source/Editor/SubobjectDataInterface/Private/SubobjectDataSubsystem.cpp
--- a/Engine/Source/Editor/SubobjectDataInterface/Private/SubobjectDataSubsystem.cpp
+++ b/Engine/Source/Editor/SubobjectDataInterface/Private/SubobjectDataSubsystem.cpp
@@ -954,10 +954,9 @@ FSubobjectDataHandle USubobjectDataSubsystem::AddNewSubobject(const FAddNewSubob
 			if (!ActorInstance)
 			{
 				const FSubobjectDataHandle& ActorRootHandle = GetActorRootHandle(ParentObjData->GetHandle());
-				if (ActorRootHandle.IsValid())
+				if (FSubobjectData* ActorInstanceData = ActorRootHandle.GetData())
 				{
-					ParentObjData = ActorRootHandle.GetData();
-					ActorInstance = ParentObjData ? ParentObjData->GetMutableActorContext() : nullptr;
+					ActorInstance = ActorInstanceData->GetMutableActorContext();
 				}
 			}
 
@@ -993,6 +992,16 @@ FSubobjectDataHandle USubobjectDataSubsystem::AddNewSubobject(const FAddNewSubob
 					else
 					{
 						USceneComponent* AttachTo = Cast<USceneComponent>(ParentObjData->GetMutableComponentTemplate());
+						if (AttachTo && AttachTo->IsTemplate())
+						{
+							UActorComponent** ArchetypeAttachment = PreInstanceComponents.FindByPredicate([AttachTo](const UActorComponent* A)
+							{ 
+								return A && A->GetArchetype() == AttachTo;
+							});
+
+							AttachTo = ArchetypeAttachment ? CastChecked<USceneComponent>(*ArchetypeAttachment) : nullptr;
+						}
+
 						if (AttachTo == nullptr)
 						{
 							AttachTo = ActorInstance->GetRootComponent();
@@ -1705,12 +1714,23 @@ bool USubobjectDataSubsystem::AttachSubobject(const FSubobjectDataHandle& OwnerH
 	
 	if(ChildToAddData->IsComponent())
 	{
-		USCS_Node* SCS_Node = OwnerData->GetSCSNode();
 		const UActorComponent* ComponentTemplate = OwnerData->GetObject<UActorComponent>();
-		
+		USceneComponent* ParentInstance = nullptr;
+		bool bIsInstancedComponent = OwnerData->IsInstancedComponent();
+
+		if (ComponentTemplate)
+		{
+			// Find the component instance on the current actor if we can
+			const FSubobjectDataHandle& ActorHandle = GetActorRootHandle(OwnerData->GetHandle());
+			const FSubobjectData* ActorData = ActorHandle.GetData();
+			ParentInstance = ActorData ? Cast<USceneComponent>(OwnerData->FindMutableComponentInstanceInActor(ActorData->GetObject<AActor>())) : nullptr;
+			bIsInstancedComponent |= ParentInstance != nullptr;
+		}
+
 		// Add a child node to the SCS tree node if not already present
 		if(USCS_Node* SCS_ChildNode = ChildToAddData->GetSCSNode())
 		{
+			USCS_Node* SCS_Node = OwnerData->GetSCSNode();
 			// Get the SCS instance that owns the child node
 			if(USimpleConstructionScript* SCS = SCS_ChildNode->GetSCS())
 			{
@@ -1749,25 +1769,22 @@ bool USubobjectDataSubsystem::AttachSubobject(const FSubobjectDataHandle& OwnerH
 				}
 			}
 		}
-		else if(OwnerData->IsInstancedComponent())
+		else if(bIsInstancedComponent)
 		{
 			USceneComponent* ChildInstance = Cast<USceneComponent>(ChildToAddData->GetMutableComponentTemplate());
-			if (ensure(ChildInstance != nullptr))
+
+			if (ensure(ChildInstance && ParentInstance))
 			{
-				USceneComponent* ParentInstance = Cast<USceneComponent>(OwnerData->GetMutableComponentTemplate());
-				if (ensure(ParentInstance != nullptr))
+				// Handle attachment at the instance level
+				if (ChildInstance->GetAttachParent() != ParentInstance)
 				{
-					// Handle attachment at the instance level
-					if (ChildInstance->GetAttachParent() != ParentInstance)
+					AActor* Owner = ParentInstance->GetOwner();
+					if (Owner->GetRootComponent() == ChildInstance)
 					{
-						AActor* Owner = ParentInstance->GetOwner();
-						if (Owner->GetRootComponent() == ChildInstance)
-						{
-							Owner->SetRootComponent(ParentInstance);
-						}
-						ChildInstance->AttachToComponent(ParentInstance, FAttachmentTransformRules::KeepWorldTransform);
+						Owner->SetRootComponent(ParentInstance);
 					}
-				}
+					ChildInstance->AttachToComponent(ParentInstance, FAttachmentTransformRules::KeepWorldTransform);
+				}	
 			}
 		}
 	}
