Replace an incorrect check for a child actor with a check for a child actor subtree object.

#rnx

diff --git a/Engine/Source/Editor/SubobjectDataInterface/Private/SubobjectDataSubsystem.cpp b/Engine/Source/Editor/SubobjectDataInterface/Private/SubobjectDataSubsystem.cpp
--- a/Engine/Source/Editor/SubobjectDataInterface/Private/SubobjectDataSubsystem.cpp
+++ b/Engine/Source/Editor/SubobjectDataInterface/Private/SubobjectDataSubsystem.cpp
@@ -1659,7 +1659,7 @@ bool USubobjectDataSubsystem::DetachSubobject(const FSubobjectDataHandle& OwnerH
 	// Remove this subobject handle from the parent
 	OwnerData->RemoveChildHandleOnly(ChildToRemoveData->GetHandle());
 	ChildToRemoveData->ClearParentHandle();
-	// if its an instance component, call detact from compont
+	// if its an instance component, call detach from component
 	if(ChildToRemoveData->IsInstancedComponent())
 	{
 		USceneComponent* ChildInstance = Cast<USceneComponent>(ChildToRemoveData->GetMutableComponentTemplate());
@@ -1672,7 +1672,7 @@ bool USubobjectDataSubsystem::DetachSubobject(const FSubobjectDataHandle& OwnerH
 	}
 
 	// Bypass removal logic if we're part of a child actor subtree
-	if (ChildToRemoveData->IsChildActor())
+	if (ChildToRemoveData->IsChildActorSubtreeObject())
 	{
 		return true;
 	}
