Fix some spelling mistakes in comments

#rnx

diff --git a/Engine/Source/Editor/SubobjectDataInterface/Public/SubobjectDataSubsystem.h b/Engine/Source/Editor/SubobjectDataInterface/Public/SubobjectDataSubsystem.h
--- a/Engine/Source/Editor/SubobjectDataInterface/Public/SubobjectDataSubsystem.h
+++ b/Engine/Source/Editor/SubobjectDataInterface/Public/SubobjectDataSubsystem.h
@@ -138,7 +138,7 @@ public:
 	* Attempt to find the subobject data for a given handle. OutData will only 
 	* be valid if the function returns true.
 	*
-	* @param Handle		Handle of the subobject data you want to aquire
+	* @param Handle		Handle of the subobject data you want to acquire
 	* @param OutData	Reference to the subobject data to populate
 	*
 	* @return bool		true if the data was found
@@ -216,7 +216,7 @@ public:
 	*
 	* @param ContextHandle			The owning context of the subobjects that should be removed
 	* @param SubobjectsToDelete		Array of subobject handles that should be deleted
-	* @param OutComponentToSelect	Populates this handle with a valid selection in the component heirarchy if desiered
+	* @param OutComponentToSelect	Populates this handle with a valid selection in the component hierarchy if desired
 	* @return 	The number of subobjects successfully deleted
 	*/
 	int32 DeleteSubobjects(const FSubobjectDataHandle& ContextHandle, const TArray<FSubobjectDataHandle>& SubobjectsToDelete, FSubobjectDataHandle& OutComponentToSelect, UBlueprint* BPContext = nullptr);
@@ -330,7 +330,7 @@ public:
 	 * Duplicate the given array of subobjects on the context. 
 	 *
 	 * @param Context			The owning context that the subobjects to dup come from
-	 * @param SubobjectsToDup	Array of handles of existing subobejcts you would like to have duplicated
+	 * @param SubobjectsToDup	Array of handles of existing subobjects you would like to have duplicated
 	 * @param BpContext			Pointer to the current blueprint context if necessary. Use nullptr if dealing with instances
 	 * @param OutNewSubobjects	Array that will be populated with any newly created subobjects
 	 */
@@ -376,7 +376,7 @@ private:
 	void FindAllSubobjectData(FSubobjectData* RootContext, TSet<FSubobjectData*>& OutVisited) const;
 	
 	/**
-	 * Gets a ref to subobject data that has been initalized with the given context.
+	 * Gets a ref to subobject data that has been initialized with the given context.
 	 * It will have a unique ID, but may be reused from the FreeData array. 
 	 */
 	FSubobjectDataHandle CreateSubobjectData(UObject* Context, const FSubobjectDataHandle& ParentHandle = FSubobjectDataHandle::InvalidHandle, bool bIsInheritedSCS = false);
