Add a flag to force deletion on subobjects so that when making a new scene root the CanDelete check will be bypassed.

This is needed because when the user drags a new scene component over the existing one, the old one should be deleted but is still marked as the SceneRoot (and cannot be deleted because we have to update the instances of the archetypes)

#rnx

diff --git a/Engine/Source/Editor/SubobjectDataInterface/Private/SubobjectDataSubsystem.cpp b/Engine/Source/Editor/SubobjectDataInterface/Private/SubobjectDataSubsystem.cpp
--- a/Engine/Source/Editor/SubobjectDataInterface/Private/SubobjectDataSubsystem.cpp
+++ b/Engine/Source/Editor/SubobjectDataInterface/Private/SubobjectDataSubsystem.cpp
@@ -1066,7 +1066,7 @@ FSubobjectDataHandle USubobjectDataSubsystem::AddNewSubobject(const FAddNewSubob
 	return NewDataHandle;
 }
 
-int32 USubobjectDataSubsystem::DeleteSubobjects(const FSubobjectDataHandle& ContextHandle, const TArray<FSubobjectDataHandle>& SubobjectsToDelete, FSubobjectDataHandle& OutComponentToSelect, UBlueprint* BPContext/* = nullptr*/)
+int32 USubobjectDataSubsystem::DeleteSubobjects(const FSubobjectDataHandle& ContextHandle, const TArray<FSubobjectDataHandle>& SubobjectsToDelete, FSubobjectDataHandle& OutComponentToSelect, UBlueprint* BPContext/* = nullptr*/, bool bForce /* = false */)
 {
 	int32 NumDeletedSubobjects = 0;
 
@@ -1088,7 +1088,7 @@ int32 USubobjectDataSubsystem::DeleteSubobjects(const FSubobjectDataHandle& Cont
 			{
 				if (const FSubobjectData* Data = Handle.GetData())
 				{
-					if (!Data->CanDelete())
+					if (!Data->CanDelete() && !bForce)
 					{
 						UE_LOG(LogSubobjectSubsystem, Warning, TEXT("Cannot delete subobject '%s' from '%s'!"), *Data->GetDisplayString(), *BPContext->GetFullName());
 						continue;
@@ -1456,7 +1456,9 @@ bool USubobjectDataSubsystem::MakeNewSceneRoot(const FSubobjectDataHandle& Conte
 
 			if (bWasDefaultSceneRoot)
 			{
-				DeleteSubobject(Context, OldSceneRoot, Blueprint);
+				TArray<FSubobjectDataHandle> ToDelete { OldSceneRoot };
+				FSubobjectDataHandle NewSelection = FSubobjectDataHandle::InvalidHandle;
+				DeleteSubobjects(Context, ToDelete, NewSelection, Blueprint, /* bForce */ true);
 			}
 		}
 	}
diff --git a/Engine/Source/Editor/SubobjectDataInterface/Public/SubobjectDataSubsystem.h b/Engine/Source/Editor/SubobjectDataInterface/Public/SubobjectDataSubsystem.h
--- a/Engine/Source/Editor/SubobjectDataInterface/Public/SubobjectDataSubsystem.h
+++ b/Engine/Source/Editor/SubobjectDataInterface/Public/SubobjectDataSubsystem.h
@@ -217,9 +217,10 @@ public:
 	* @param ContextHandle			The owning context of the subobjects that should be removed
 	* @param SubobjectsToDelete		Array of subobject handles that should be deleted
 	* @param OutComponentToSelect	Populates this handle with a valid selection in the component hierarchy if desired
+	* @param bForce					If true then this will attempt to delete a subobject even if the CanDelete function flags otherwise.
 	* @return 	The number of subobjects successfully deleted
 	*/
-	int32 DeleteSubobjects(const FSubobjectDataHandle& ContextHandle, const TArray<FSubobjectDataHandle>& SubobjectsToDelete, FSubobjectDataHandle& OutComponentToSelect, UBlueprint* BPContext = nullptr);
+	int32 DeleteSubobjects(const FSubobjectDataHandle& ContextHandle, const TArray<FSubobjectDataHandle>& SubobjectsToDelete, FSubobjectDataHandle& OutComponentToSelect, UBlueprint* BPContext = nullptr, bool bForce = false);
 
 	/**
 	* Attempts to delete the given subobject from its blueprint context 
