Add an additional check when determining the default scene root to see if a native subobject is the root of it's owning actor.

#rnx

diff --git a/Engine/Source/Editor/SubobjectDataInterface/Private/SubobjectData.cpp b/Engine/Source/Editor/SubobjectDataInterface/Private/SubobjectData.cpp
--- a/Engine/Source/Editor/SubobjectDataInterface/Private/SubobjectData.cpp
+++ b/Engine/Source/Editor/SubobjectDataInterface/Private/SubobjectData.cpp
@@ -902,6 +902,13 @@ bool FSubobjectData::IsDefaultSceneRoot() const
 			return SCS_Node == SCS->GetDefaultSceneRootNode();
 		}
 	}
+	// As a last resort check the owning actor to see if the root component matches up with this
+	// This will be the case for native subobjects
+	else if(SceneComponent)
+	{
+		const AActor* Owner = SceneComponent->GetOwner();
+		return Owner && Owner->GetRootComponent() == SceneComponent;
+	}
 
 	// Nothing else can be a DefaultSceneRoot
 	return false;
