CIS fix. Remove ParentObjData from an if statement because there is already a "check" on the variable above. Having it in the if statement was causing other uses of it to flag PVS

#rnx

diff --git a/Engine/Source/Editor/SubobjectDataInterface/Private/SubobjectDataSubsystem.cpp b/Engine/Source/Editor/SubobjectDataInterface/Private/SubobjectDataSubsystem.cpp
--- a/Engine/Source/Editor/SubobjectDataInterface/Private/SubobjectDataSubsystem.cpp
+++ b/Engine/Source/Editor/SubobjectDataInterface/Private/SubobjectDataSubsystem.cpp
@@ -960,7 +960,7 @@ FSubobjectDataHandle USubobjectDataSubsystem::AddNewSubobject(const FAddNewSubob
 				}
 			}
 
-			if (ActorInstance && ParentObjData)
+			if (ActorInstance)
 			{
 				ActorInstance->Modify();
 
