Check if the old scene root was the default scene root that is generated by every BP to determine if it should be deleted or not. This makes the behavior consistent with the old SCS Editor

#rnx

diff --git a/Engine/Source/Editor/SubobjectDataInterface/Private/SubobjectDataSubsystem.cpp b/Engine/Source/Editor/SubobjectDataInterface/Private/SubobjectDataSubsystem.cpp
--- a/Engine/Source/Editor/SubobjectDataInterface/Private/SubobjectDataSubsystem.cpp
+++ b/Engine/Source/Editor/SubobjectDataInterface/Private/SubobjectDataSubsystem.cpp
@@ -1453,8 +1453,11 @@ bool USubobjectDataSubsystem::MakeNewSceneRoot(const FSubobjectDataHandle& Conte
 
 			// Set old root as child of new root
 			AttachSubobject(DroppedNewSceneRootData->GetHandle(), OldSceneRoot);
+			FSubobjectData* OldData = OldSceneRoot.GetData();
 
-			if (bWasDefaultSceneRoot)
+			// If the old SCS node was the default scene root created by every blueprint, then we can delete it
+			// Otherwise we will keep it and simply reparent it to the new scene root
+			if (OldData->GetSCSNode() == Blueprint->SimpleConstructionScript->GetDefaultSceneRootNode())
 			{
 				TArray<FSubobjectDataHandle> ToDelete { OldSceneRoot };
 				FSubobjectDataHandle NewSelection = FSubobjectDataHandle::InvalidHandle;
