Check not only the SimpleConstructionScript's DefaultSceneRootNode, but also the SceneRootComponentTemplate when determining if a subobject is the default scene root. This handles cases where the DefaultSceneRoot node has been updated (like when you drag/drop a new scene component onto it in the Subobject Editor)

#rnx

diff --git a/Engine/Source/Editor/SubobjectDataInterface/Private/SubobjectData.cpp b/Engine/Source/Editor/SubobjectDataInterface/Private/SubobjectData.cpp
--- a/Engine/Source/Editor/SubobjectDataInterface/Private/SubobjectData.cpp
+++ b/Engine/Source/Editor/SubobjectDataInterface/Private/SubobjectData.cpp
@@ -896,10 +896,11 @@ bool FSubobjectData::IsDefaultSceneRoot() const
 	// If this isn't a scene component, then we can check an SCS node for a flag. 
 	else if (const USCS_Node* SCS_Node = GetSCSNode())
 	{
-		USimpleConstructionScript* SCS = SCS_Node->GetSCS();
-		if (SCS != nullptr)
+		if (USimpleConstructionScript* SCS = SCS_Node->GetSCS())
 		{
-			return SCS_Node == SCS->GetDefaultSceneRootNode();
+			const USceneComponent* SCS_Root = SCS->GetSceneRootComponentTemplate();
+
+			return (SCS_Node == SCS->GetDefaultSceneRootNode()) || (SceneComponent && (SCS_Root == SceneComponent));
 		}
 	}
 	// As a last resort check the owning actor to see if the root component matches up with this
