Add a fallback method of determining the instanced component to use as an asset override in the event that a subobjec thas been added but is not a part of the generated class yet. Do this by checking if the archetype of the GEN_VARIABLE subobject and one on the instance are the same.

#rnx

diff --git a/Engine/Source/Editor/SubobjectDataInterface/Private/SubobjectData.cpp b/Engine/Source/Editor/SubobjectDataInterface/Private/SubobjectData.cpp
--- a/Engine/Source/Editor/SubobjectDataInterface/Private/SubobjectData.cpp
+++ b/Engine/Source/Editor/SubobjectDataInterface/Private/SubobjectData.cpp
@@ -190,6 +190,19 @@ UActorComponent* FSubobjectData::FindMutableComponentInstanceInActor(const AActo
 			InActor->GetComponents(Components);
 			ComponentInstance = FComponentEditorUtils::FindMatchingComponent(ComponentTemplate, Components); 
 		}
+
+		if (!ComponentInstance && SCS_Node)
+		{
+			TInlineComponentArray<UActorComponent*> Components;
+			InActor->GetComponents(Components);
+
+			UActorComponent** MatchingArchetype = Components.FindByPredicate([SCS_Node](const UActorComponent* A)
+			{
+				return A && A->GetArchetype() == SCS_Node->ComponentTemplate;
+			});
+
+			ComponentInstance = MatchingArchetype ? *MatchingArchetype : nullptr;
+		}
 	}
 
 	return ComponentInstance;
