Wrapped ObjectPtr upgrade refresh for Engine + ShooterGame

diff --git a/Engine/Source/Developer/RigVMDeveloper/Public/RigVMModel/RigVMFunctionLibrary.h b/Engine/Source/Developer/RigVMDeveloper/Public/RigVMModel/RigVMFunctionLibrary.h
--- a/Engine/Source/Developer/RigVMDeveloper/Public/RigVMModel/RigVMFunctionLibrary.h
+++ b/Engine/Source/Developer/RigVMDeveloper/Public/RigVMModel/RigVMFunctionLibrary.h
@@ -92,13 +92,13 @@ public:
 private:
 
 	UPROPERTY()
-	TMap< URigVMLibraryNode*, FRigVMFunctionReferenceArray > FunctionReferences;
+	TMap< TObjectPtr<URigVMLibraryNode>, FRigVMFunctionReferenceArray > FunctionReferences;
 
 	// A map which stores a library node per original pathname.
 	// The source pathname is the full path of the source function that was localized
 	// to the local copy stored in the value of the pair.
 	UPROPERTY()
-	TMap< FString, URigVMLibraryNode* > LocalizedFunctions;
+	TMap< FString, TObjectPtr<URigVMLibraryNode> > LocalizedFunctions;
 
 	friend class URigVMController;
 	friend class URigVMCompiler;
diff --git a/Engine/Source/Developer/RigVMDeveloper/Public/RigVMModel/RigVMGraph.h b/Engine/Source/Developer/RigVMDeveloper/Public/RigVMModel/RigVMGraph.h
--- a/Engine/Source/Developer/RigVMDeveloper/Public/RigVMModel/RigVMGraph.h
+++ b/Engine/Source/Developer/RigVMDeveloper/Public/RigVMModel/RigVMGraph.h
@@ -140,10 +140,10 @@ private:
 	void Notify(ERigVMGraphNotifType InNotifType, UObject* InSubject);
 
 	UPROPERTY()
-	TArray<URigVMNode*> Nodes;
+	TArray<TObjectPtr<URigVMNode>> Nodes;
 
 	UPROPERTY()
-	TArray<URigVMLink*> Links;
+	TArray<TObjectPtr<URigVMLink>> Links;
 
 	UPROPERTY()
 	TArray<FName> SelectedNodes;
diff --git a/Engine/Source/Developer/RigVMDeveloper/Public/RigVMModel/RigVMNode.h b/Engine/Source/Developer/RigVMDeveloper/Public/RigVMModel/RigVMNode.h
--- a/Engine/Source/Developer/RigVMDeveloper/Public/RigVMModel/RigVMNode.h
+++ b/Engine/Source/Developer/RigVMDeveloper/Public/RigVMModel/RigVMNode.h
@@ -241,10 +241,10 @@ protected:
 private:
 
 	UPROPERTY()
-	TArray<URigVMPin*> Pins;
+	TArray<TObjectPtr<URigVMPin>> Pins;
 
 	UPROPERTY()
-	TArray<URigVMPin*> OrphanedPins;
+	TArray<TObjectPtr<URigVMPin>> OrphanedPins;
 
 	friend class URigVMController;
 	friend class URigVMGraph;
diff --git a/Engine/Source/Developer/RigVMDeveloper/Public/RigVMModel/RigVMPin.h b/Engine/Source/Developer/RigVMDeveloper/Public/RigVMModel/RigVMPin.h
--- a/Engine/Source/Developer/RigVMDeveloper/Public/RigVMModel/RigVMPin.h
+++ b/Engine/Source/Developer/RigVMDeveloper/Public/RigVMModel/RigVMPin.h
@@ -33,16 +33,16 @@ public:
 	}
 
 	UPROPERTY()
-	URigVMUnitNode* UnitNode;
+	TObjectPtr<URigVMUnitNode> UnitNode;
 
 	UPROPERTY()
 	bool bInjectedAsInput;
 
 	UPROPERTY()
-	URigVMPin* InputPin;
+	TObjectPtr<URigVMPin> InputPin;
 
 	UPROPERTY()
-	URigVMPin* OutputPin;
+	TObjectPtr<URigVMPin> OutputPin;
 
 	// Returns the graph of this injected node.
 	UFUNCTION(BlueprintCallable, Category = RigVMInjectionInfo)
@@ -393,7 +393,7 @@ private:
 	FString CPPType;
 
 	UPROPERTY(transient)
-	UObject* CPPTypeObject;
+	TObjectPtr<UObject> CPPTypeObject;
 
 	UPROPERTY()
 	FName CPPTypeObjectPath;
@@ -405,13 +405,13 @@ private:
 	FName CustomWidgetName;
 
 	UPROPERTY()
-	TArray<URigVMPin*> SubPins;
+	TArray<TObjectPtr<URigVMPin>> SubPins;
 
 	UPROPERTY(transient)
-	TArray<URigVMLink*> Links;
+	TArray<TObjectPtr<URigVMLink>> Links;
 
 	UPROPERTY()
-	TArray<URigVMInjectionInfo*> InjectionInfos;
+	TArray<TObjectPtr<URigVMInjectionInfo>> InjectionInfos;
 
 	UPROPERTY()
 	FString BoundVariablePath;
diff --git a/Engine/Source/Developer/RigVMDeveloper/Public/RigVMModel/RigVMVariableDescription.h b/Engine/Source/Developer/RigVMDeveloper/Public/RigVMModel/RigVMVariableDescription.h
--- a/Engine/Source/Developer/RigVMDeveloper/Public/RigVMModel/RigVMVariableDescription.h
+++ b/Engine/Source/Developer/RigVMDeveloper/Public/RigVMModel/RigVMVariableDescription.h
@@ -35,7 +35,7 @@ public:
 
 	// The Struct of the C++ data type of the variable (or nullptr)
 	UPROPERTY(BlueprintReadOnly, Category = RigVMGraphVariableDescription)
-	UObject* CPPTypeObject = nullptr;
+	TObjectPtr<UObject> CPPTypeObject = nullptr;
 
 	// The default value of the variable
 	UPROPERTY(BlueprintReadOnly, Category = RigVMGraphVariableDescription)
diff --git a/Engine/Source/Editor/AnimGraph/Private/AnimBlueprintExtension_BlendSpaceGraph.h b/Engine/Source/Editor/AnimGraph/Private/AnimBlueprintExtension_BlendSpaceGraph.h
--- a/Engine/Source/Editor/AnimGraph/Private/AnimBlueprintExtension_BlendSpaceGraph.h
+++ b/Engine/Source/Editor/AnimGraph/Private/AnimBlueprintExtension_BlendSpaceGraph.h
@@ -26,7 +26,7 @@ class UAnimBlueprintExtension_BlendSpaceGraph : public UAnimBlueprintExtension
 private:
 	// The class that is being compiled
 	UPROPERTY(Transient)
-	UClass* Class;
+	TObjectPtr<UClass> Class;
 
 private:
 	UPROPERTY()
diff --git a/Engine/Source/Editor/AnimGraph/Private/AnimGraphNode_CallFunction.h b/Engine/Source/Editor/AnimGraph/Private/AnimGraphNode_CallFunction.h
--- a/Engine/Source/Editor/AnimGraph/Private/AnimGraphNode_CallFunction.h
+++ b/Engine/Source/Editor/AnimGraph/Private/AnimGraphNode_CallFunction.h
@@ -20,11 +20,11 @@ class UAnimGraphNode_CallFunction : public UAnimGraphNode_Base
 private:
 	// Inner graph we maintain to hook into the function call machinery
 	UPROPERTY()
-	UEdGraph* InnerGraph;
+	TObjectPtr<UEdGraph> InnerGraph;
 
 	// Inner node we maintain to hook into the function call machinery
 	UPROPERTY()
-	UK2Node_CallFunction* CallFunctionPrototype;
+	TObjectPtr<UK2Node_CallFunction> CallFunctionPrototype;
 	
     UPROPERTY(EditAnywhere, Category=Settings)
 	FAnimNode_CallFunction Node;
diff --git a/Engine/Source/Editor/Experimental/EditorInteractiveToolsFramework/Public/EditorGizmos/EditorTransformGizmo.h b/Engine/Source/Editor/Experimental/EditorInteractiveToolsFramework/Public/EditorGizmos/EditorTransformGizmo.h
--- a/Engine/Source/Editor/Experimental/EditorInteractiveToolsFramework/Public/EditorGizmos/EditorTransformGizmo.h
+++ b/Engine/Source/Editor/Experimental/EditorInteractiveToolsFramework/Public/EditorGizmos/EditorTransformGizmo.h
@@ -124,7 +124,7 @@ public:
 
 	/** The active target object for the Gizmo */
 	UPROPERTY()
-	UTransformProxy* ActiveTarget;
+	TObjectPtr<UTransformProxy> ActiveTarget;
 
 	/**
 	 * Repositions the gizmo without issuing undo/redo changes, triggering callbacks, 
diff --git a/Engine/Source/Editor/SubobjectDataInterface/Public/SubobjectDataSubsystem.h b/Engine/Source/Editor/SubobjectDataInterface/Public/SubobjectDataSubsystem.h
--- a/Engine/Source/Editor/SubobjectDataInterface/Public/SubobjectDataSubsystem.h
+++ b/Engine/Source/Editor/SubobjectDataInterface/Public/SubobjectDataSubsystem.h
@@ -33,7 +33,7 @@ struct FAddNewSubobjectParams
 
 	/** The class of the new subobject that will be added */
 	UPROPERTY(BlueprintReadWrite, Category = "SubobjectDataSubsystem")
-	UClass* NewClass;
+	TObjectPtr<UClass> NewClass;
 
 	/** Specific asset to use instead of the selected asset in the content browser */
 	UObject* AssetOverride;
@@ -43,7 +43,7 @@ struct FAddNewSubobjectParams
 	 * this subobject is being added to an instance.
 	 */
 	UPROPERTY(BlueprintReadWrite, Category = "SubobjectDataSubsystem")
-	UBlueprint* BlueprintContext;
+	TObjectPtr<UBlueprint> BlueprintContext;
 	
 	/** Optionally skip marking this blueprint as modified (e.g. if we're handling that externally */
 	UPROPERTY(BlueprintReadWrite, Category = "SubobjectDataSubsystem")
@@ -71,14 +71,14 @@ struct FReparentSubobjectParams
 	* this subobject is being added to an instance.
 	*/
 	UPROPERTY(BlueprintReadWrite, Category = "SubobjectDataSubsystem")
-	UBlueprint* BlueprintContext = nullptr;
+	TObjectPtr<UBlueprint> BlueprintContext = nullptr;
 
 	/** 
 	* The preview actor context to be used if in a blueprint context. 
 	* This must have a value if BlueprintContext is needed. 
 	*/
 	UPROPERTY(BlueprintReadWrite, Category = "SubobjectDataSubsystem")
-	AActor* ActorPreviewContext = nullptr;
+	TObjectPtr<AActor> ActorPreviewContext = nullptr;
 };
 
 /**
diff --git a/Engine/Source/Editor/UMGEditor/Public/WidgetBlueprintThumbnailRenderer.h b/Engine/Source/Editor/UMGEditor/Public/WidgetBlueprintThumbnailRenderer.h
--- a/Engine/Source/Editor/UMGEditor/Public/WidgetBlueprintThumbnailRenderer.h
+++ b/Engine/Source/Editor/UMGEditor/Public/WidgetBlueprintThumbnailRenderer.h
@@ -23,6 +23,6 @@ private:
 	TTuple<FVector2D, FVector2D> GetScaledSizeAndOffset(float ImgWidth, float ImgHeight, float ThumbnailWidth, float ThumbnailHeight) const;
 
 	UPROPERTY(Transient)
-	UTextureRenderTarget2D* RenderTarget2D;
+	TObjectPtr<UTextureRenderTarget2D> RenderTarget2D;
 
 };
diff --git a/Engine/Source/Editor/UnrealEd/Classes/Editor/EditorEngine.h b/Engine/Source/Editor/UnrealEd/Classes/Editor/EditorEngine.h
--- a/Engine/Source/Editor/UnrealEd/Classes/Editor/EditorEngine.h
+++ b/Engine/Source/Editor/UnrealEd/Classes/Editor/EditorEngine.h
@@ -2435,7 +2435,7 @@ protected:
 	virtual ULevelEditorDragDropHandler* CreateLevelEditorDragDropHandler();
 private:
 	UPROPERTY()
-	mutable ULevelEditorDragDropHandler* DragDropHandler;
+	mutable TObjectPtr<ULevelEditorDragDropHandler> DragDropHandler;
 
 protected:
 
diff --git a/Engine/Source/Editor/UnrealEd/Classes/Factories/MirrorDataTableFactory.h b/Engine/Source/Editor/UnrealEd/Classes/Factories/MirrorDataTableFactory.h
--- a/Engine/Source/Editor/UnrealEd/Classes/Factories/MirrorDataTableFactory.h
+++ b/Engine/Source/Editor/UnrealEd/Classes/Factories/MirrorDataTableFactory.h
@@ -28,13 +28,13 @@ class UNREALED_API UMirrorDataTableFactory : public UFactory
 	UMirrorDataTableFactory(const FObjectInitializer& ObjectInitializer);
 
 	UPROPERTY()
-	const class UScriptStruct* Struct;
+	TObjectPtr<const class UScriptStruct> Struct;
 
 	UPROPERTY()
-	class USkeleton* Skeleton;
+	TObjectPtr<class USkeleton> Skeleton;
 
 	UPROPERTY()
-	class UMirrorTableFindReplaceExpressions* MirrorFindReplaceExpressions;
+	TObjectPtr<class UMirrorTableFindReplaceExpressions> MirrorFindReplaceExpressions;
 
 public:
 	//~ Begin UFactory Interface
diff --git a/Engine/Source/Runtime/AudioMixer/Public/Quartz/AudioMixerClockHandle.h b/Engine/Source/Runtime/AudioMixer/Public/Quartz/AudioMixerClockHandle.h
--- a/Engine/Source/Runtime/AudioMixer/Public/Quartz/AudioMixerClockHandle.h
+++ b/Engine/Source/Runtime/AudioMixer/Public/Quartz/AudioMixerClockHandle.h
@@ -181,7 +181,7 @@ private:
 	MetronomeDelegateGameThreadData MetronomeDelegates[static_cast<int32>(EQuartzCommandQuantization::Count)];
 
 	UPROPERTY(Transient)
-	UQuartzSubsystem* QuartzSubsystem;
+	TObjectPtr<UQuartzSubsystem> QuartzSubsystem;
 
 	FName ClockHandleId;
 
@@ -190,6 +190,6 @@ private:
 	bool bConnectedToClock{ false };
 
 	UPROPERTY(Transient)
-	UWorld* WorldPtr{ nullptr };
+	TObjectPtr<UWorld> WorldPtr{ nullptr };
 
 }; // class UClockHandleBlueprintLibrary
diff --git a/Engine/Source/Runtime/ClothingSystemRuntimeInterface/Public/ClothingSimulationInteractor.h b/Engine/Source/Runtime/ClothingSystemRuntimeInterface/Public/ClothingSimulationInteractor.h
--- a/Engine/Source/Runtime/ClothingSystemRuntimeInterface/Public/ClothingSimulationInteractor.h
+++ b/Engine/Source/Runtime/ClothingSystemRuntimeInterface/Public/ClothingSimulationInteractor.h
@@ -129,7 +129,7 @@ public:
 
 	/** Cloth interactors currently created. */
 	UPROPERTY()
-	TMap<FName, UClothingInteractor*> ClothingInteractors;
+	TMap<FName, TObjectPtr<UClothingInteractor>> ClothingInteractors;
 
 protected:
 	virtual UClothingInteractor* CreateClothingInteractor()
diff --git a/Engine/Source/Runtime/Engine/Classes/Animation/AnimBlueprintGeneratedClass.h b/Engine/Source/Runtime/Engine/Classes/Animation/AnimBlueprintGeneratedClass.h
--- a/Engine/Source/Runtime/Engine/Classes/Animation/AnimBlueprintGeneratedClass.h
+++ b/Engine/Source/Runtime/Engine/Classes/Animation/AnimBlueprintGeneratedClass.h
@@ -410,7 +410,7 @@ private:
 
 	// Map from anim node struct to info about that struct (used to accelerate property name lookups)
 	UPROPERTY()
-	TMap<const UScriptStruct*, FAnimNodeStructData> NodeTypeMap;
+	TMap<TObjectPtr<const UScriptStruct>, FAnimNodeStructData> NodeTypeMap;
 
 	// Cached properties used to access 'folded' anim node properties
 	TArray<FProperty*> MutableProperties;
diff --git a/Engine/Source/Runtime/Engine/Classes/Animation/AnimData/AttributeIdentifier.h b/Engine/Source/Runtime/Engine/Classes/Animation/AnimData/AttributeIdentifier.h
--- a/Engine/Source/Runtime/Engine/Classes/Animation/AnimData/AttributeIdentifier.h
+++ b/Engine/Source/Runtime/Engine/Classes/Animation/AnimData/AttributeIdentifier.h
@@ -74,7 +74,7 @@ protected:
 	int32 BoneIndex;
 
 	UPROPERTY(BlueprintReadOnly, Transient, Category = AttributeIdentifier)
-	UScriptStruct* ScriptStruct;
+	TObjectPtr<UScriptStruct> ScriptStruct;
 
 	UPROPERTY(BlueprintReadOnly, Category = AttributeIdentifier)
 	FSoftObjectPath ScriptStructPath;
diff --git a/Engine/Source/Runtime/Engine/Classes/Animation/AnimNodeFunctionRef.h b/Engine/Source/Runtime/Engine/Classes/Animation/AnimNodeFunctionRef.h
--- a/Engine/Source/Runtime/Engine/Classes/Animation/AnimNodeFunctionRef.h
+++ b/Engine/Source/Runtime/Engine/Classes/Animation/AnimNodeFunctionRef.h
@@ -52,7 +52,7 @@ private:
 
 	// The function to call, recovered by looking for a function of name FunctionName
 	UPROPERTY(Transient)
-	UFunction* Function = nullptr;
+	TObjectPtr<UFunction> Function = nullptr;
 };
 
 namespace UE { namespace Anim {
diff --git a/Engine/Source/Runtime/Engine/Classes/Engine/Level.h b/Engine/Source/Runtime/Engine/Classes/Engine/Level.h
--- a/Engine/Source/Runtime/Engine/Classes/Engine/Level.h
+++ b/Engine/Source/Runtime/Engine/Classes/Engine/Level.h
@@ -88,7 +88,7 @@ class UActorContainer : public UObject
 
 public:
 	UPROPERTY()
-	TMap<FName, AActor*> Actors;
+	TMap<FName, TObjectPtr<AActor>> Actors;
 };
 
 /**
@@ -457,7 +457,7 @@ public:
 
 	/** List of modified, unsaved actors that needs to be duplicated for PIE */
 	UPROPERTY(NonPIEDuplicateTransient)
-	TMap<FName, AActor*> ActorsModifiedForPIE;
+	TMap<FName, TObjectPtr<AActor>> ActorsModifiedForPIE;
 
 	/** Use external actors, new actor spawned in this level will be external and existing external actors will be loaded on load. */
 	UPROPERTY(EditInstanceOnly, Category=World)
diff --git a/Engine/Source/Runtime/Engine/Classes/Engine/StaticMeshSourceData.h b/Engine/Source/Runtime/Engine/Classes/Engine/StaticMeshSourceData.h
--- a/Engine/Source/Runtime/Engine/Classes/Engine/StaticMeshSourceData.h
+++ b/Engine/Source/Runtime/Engine/Classes/Engine/StaticMeshSourceData.h
@@ -78,7 +78,7 @@ public:
 	 * If the bulk data within is empty, the LOD is autogenerated (for LOD1+).
 	 */
 	UPROPERTY()
-	UStaticMeshDescriptionBulkData* StaticMeshDescriptionBulkData;
+	TObjectPtr<UStaticMeshDescriptionBulkData> StaticMeshDescriptionBulkData;
 #endif
 
 #if WITH_EDITOR
diff --git a/Engine/Source/Runtime/Engine/Classes/Materials/Material.h b/Engine/Source/Runtime/Engine/Classes/Materials/Material.h
--- a/Engine/Source/Runtime/Engine/Classes/Materials/Material.h
+++ b/Engine/Source/Runtime/Engine/Classes/Materials/Material.h
@@ -848,10 +848,10 @@ public:
 
 	/** The execution begin expression, if material is using an exec wire */
 	UPROPERTY()
-	class UMaterialExpressionExecBegin* ExpressionExecBegin;
+	TObjectPtr<class UMaterialExpressionExecBegin> ExpressionExecBegin;
 	
 	UPROPERTY()
-	class UMaterialExpressionExecEnd* ExpressionExecEnd;
+	TObjectPtr<class UMaterialExpressionExecEnd> ExpressionExecEnd;
 
 	/** Controls where this parameter group is displayed in a material instance parameter list.  The lower the number the higher up in the parameter list. */
 	UPROPERTY(EditAnywhere, EditFixedSize, Category = "Group Sorting")
diff --git a/Engine/Source/Runtime/Engine/Classes/Materials/MaterialExpression.h b/Engine/Source/Runtime/Engine/Classes/Materials/MaterialExpression.h
--- a/Engine/Source/Runtime/Engine/Classes/Materials/MaterialExpression.h
+++ b/Engine/Source/Runtime/Engine/Classes/Materials/MaterialExpression.h
@@ -134,7 +134,7 @@ public:
 #endif // WITH_EDITOR
 private:
 	UPROPERTY()
-	class UMaterialExpression* Expression = nullptr;
+	TObjectPtr<class UMaterialExpression> Expression = nullptr;
 };
 
 struct FExpressionExecOutputEntry
diff --git a/Engine/Source/Runtime/Engine/Classes/Sound/AudioSettings.h b/Engine/Source/Runtime/Engine/Classes/Sound/AudioSettings.h
--- a/Engine/Source/Runtime/Engine/Classes/Sound/AudioSettings.h
+++ b/Engine/Source/Runtime/Engine/Classes/Sound/AudioSettings.h
@@ -225,13 +225,13 @@ class ENGINE_API UAudioSettings : public UDeveloperSettings
 
 private:
 	UPROPERTY(Transient)
-	USoundClass* DefaultSoundClass;
+	TObjectPtr<USoundClass> DefaultSoundClass;
 
 	UPROPERTY(Transient)
-	USoundClass* DefaultMediaSoundClass;
+	TObjectPtr<USoundClass> DefaultMediaSoundClass;
 
 	UPROPERTY(Transient)
-	USoundConcurrency* DefaultSoundConcurrency;
+	TObjectPtr<USoundConcurrency> DefaultSoundConcurrency;
 
 public:
 	void LoadDefaultObjects();
diff --git a/Engine/Source/Runtime/Engine/Public/ISMPartition/ISMComponentData.h b/Engine/Source/Runtime/Engine/Public/ISMPartition/ISMComponentData.h
--- a/Engine/Source/Runtime/Engine/Public/ISMPartition/ISMComponentData.h
+++ b/Engine/Source/Runtime/Engine/Public/ISMPartition/ISMComponentData.h
@@ -88,7 +88,7 @@ struct FISMComponentData
 	TArray<FISMClientData> ClientInstances;
 
 	UPROPERTY()
-	UInstancedStaticMeshComponent* Component;
+	TObjectPtr<UInstancedStaticMeshComponent> Component;
 
 	/** If Lighting cache should be invalidated at the end of a BeginUpdate/EndUpdate */
 	bool bInvalidateLightingCache;
diff --git a/Engine/Source/Runtime/Engine/Public/ISMPartition/ISMComponentDescriptor.h b/Engine/Source/Runtime/Engine/Public/ISMPartition/ISMComponentDescriptor.h
--- a/Engine/Source/Runtime/Engine/Public/ISMPartition/ISMComponentDescriptor.h
+++ b/Engine/Source/Runtime/Engine/Public/ISMPartition/ISMComponentDescriptor.h
@@ -56,13 +56,13 @@ public:
 	TSubclassOf<UInstancedStaticMeshComponent> ComponentClass;
 	
 	UPROPERTY()
-	UStaticMesh* StaticMesh = nullptr;
+	TObjectPtr<UStaticMesh> StaticMesh = nullptr;
 	
 	UPROPERTY()
-	TArray<UMaterialInterface*> OverrideMaterials;
+	TArray<TObjectPtr<UMaterialInterface>> OverrideMaterials;
 	
 	UPROPERTY()
-	TArray<URuntimeVirtualTexture*> RuntimeVirtualTextures;
+	TArray<TObjectPtr<URuntimeVirtualTexture>> RuntimeVirtualTextures;
 
 	UPROPERTY()
 	TEnumAsByte<EComponentMobility::Type> Mobility;
diff --git a/Engine/Source/Runtime/Engine/Public/WorldPartition/WorldPartitionLevelStreamingDynamic.h b/Engine/Source/Runtime/Engine/Public/WorldPartition/WorldPartitionLevelStreamingDynamic.h
--- a/Engine/Source/Runtime/Engine/Public/WorldPartition/WorldPartitionLevelStreamingDynamic.h
+++ b/Engine/Source/Runtime/Engine/Public/WorldPartition/WorldPartitionLevelStreamingDynamic.h
@@ -68,7 +68,7 @@ private:
 
 #if WITH_EDITORONLY_DATA
 	UPROPERTY()
-	UActorContainer* ActorContainer;
+	TObjectPtr<UActorContainer> ActorContainer;
 #endif
 
 	UPROPERTY()
diff --git a/Engine/Source/Runtime/Engine/Public/WorldPartition/WorldPartitionRuntimeSpatialHashCell.h b/Engine/Source/Runtime/Engine/Public/WorldPartition/WorldPartitionRuntimeSpatialHashCell.h
--- a/Engine/Source/Runtime/Engine/Public/WorldPartition/WorldPartitionRuntimeSpatialHashCell.h
+++ b/Engine/Source/Runtime/Engine/Public/WorldPartition/WorldPartitionRuntimeSpatialHashCell.h
@@ -35,6 +35,6 @@ class UWorldPartitionRuntimeSpatialHashCell : public UWorldPartitionRuntimeCell
 
 #if WITH_EDITORONLY_DATA
 	UPROPERTY()
-	UActorContainer* ActorContainer;
+	TObjectPtr<UActorContainer> ActorContainer;
 #endif
 };
\ No newline at end of file
diff --git a/Engine/Source/Runtime/Engine/Public/WorldPartition/WorldPartitionRuntimeSpatialHashGridPreviewer.h b/Engine/Source/Runtime/Engine/Public/WorldPartition/WorldPartitionRuntimeSpatialHashGridPreviewer.h
--- a/Engine/Source/Runtime/Engine/Public/WorldPartition/WorldPartitionRuntimeSpatialHashGridPreviewer.h
+++ b/Engine/Source/Runtime/Engine/Public/WorldPartition/WorldPartitionRuntimeSpatialHashGridPreviewer.h
@@ -22,13 +22,13 @@ struct ENGINE_API FWorldPartitionRuntimeSpatialHashGridPreviewer
 
 #if WITH_EDITORONLY_DATA
 	UPROPERTY()
-	UMaterial* Material;
+	TObjectPtr<UMaterial> Material;
 
 	UPROPERTY()
-	UMaterialInstanceDynamic* MID;
+	TObjectPtr<UMaterialInstanceDynamic> MID;
 
 	UPROPERTY()
-	APostProcessVolume* Volume;
+	TObjectPtr<APostProcessVolume> Volume;
 
 	struct FGridParametersCache
 	{
diff --git a/Engine/Source/Runtime/Experimental/FieldSystem/Source/FieldSystemEngine/Public/Field/FieldSystemObjects.h b/Engine/Source/Runtime/Experimental/FieldSystem/Source/FieldSystemEngine/Public/Field/FieldSystemObjects.h
--- a/Engine/Source/Runtime/Experimental/FieldSystem/Source/FieldSystemEngine/Public/Field/FieldSystemObjects.h
+++ b/Engine/Source/Runtime/Experimental/FieldSystem/Source/FieldSystemEngine/Public/Field/FieldSystemObjects.h
@@ -1020,11 +1020,11 @@ struct FIELDSYSTEMENGINE_API FFieldObjectCommands
 
 	/** Commands Root Node */
 	UPROPERTY(BlueprintReadWrite, EditAnywhere, Category = "Field", DisplayName = "Root Nodes")
-	TArray<UFieldNodeBase*> RootNodes;
+	TArray<TObjectPtr<UFieldNodeBase>> RootNodes;
 
 	/** Commands Meta Data*/
 	UPROPERTY(BlueprintReadWrite, EditAnywhere, Category = "Field", DisplayName = "Meta Datas")
-	TArray<UFieldSystemMetaData*> MetaDatas;
+	TArray<TObjectPtr<UFieldSystemMetaData>> MetaDatas;
 };
 
 
diff --git a/Engine/Source/Runtime/Experimental/GeometryCollectionEngine/Public/GeometryCollection/GeometryCollectionComponent.h b/Engine/Source/Runtime/Experimental/GeometryCollectionEngine/Public/GeometryCollection/GeometryCollectionComponent.h
--- a/Engine/Source/Runtime/Experimental/GeometryCollectionEngine/Public/GeometryCollection/GeometryCollectionComponent.h
+++ b/Engine/Source/Runtime/Experimental/GeometryCollectionEngine/Public/GeometryCollection/GeometryCollectionComponent.h
@@ -800,7 +800,7 @@ private:
 
 	/** The information of all the embedded instanced static meshes */
 	UPROPERTY()
-	TArray<UInstancedStaticMeshComponent*> EmbeddedGeometryComponents;
+	TArray<TObjectPtr<UInstancedStaticMeshComponent>> EmbeddedGeometryComponents;
 
 	bool IsEmbeddedGeometryValid() const;
 	void ClearEmbeddedGeometry();
diff --git a/Engine/Source/Runtime/Experimental/GeometryFramework/Public/UDynamicMesh.h b/Engine/Source/Runtime/Experimental/GeometryFramework/Public/UDynamicMesh.h
--- a/Engine/Source/Runtime/Experimental/GeometryFramework/Public/UDynamicMesh.h
+++ b/Engine/Source/Runtime/Experimental/GeometryFramework/Public/UDynamicMesh.h
@@ -388,9 +388,9 @@ public:
 protected:
 	/** Meshes in the pool that are available */
 	UPROPERTY()
-	TArray<UDynamicMesh*> CachedMeshes;
+	TArray<TObjectPtr<UDynamicMesh>> CachedMeshes;
 
 	/** All meshes the pool has allocated */
 	UPROPERTY()
-	TArray<UDynamicMesh*> AllCreatedMeshes;
+	TArray<TObjectPtr<UDynamicMesh>> AllCreatedMeshes;
 };
\ No newline at end of file
diff --git a/Engine/Source/Runtime/Experimental/InteractiveToolsFramework/Public/BaseGizmos/AxisPositionGizmo.h b/Engine/Source/Runtime/Experimental/InteractiveToolsFramework/Public/BaseGizmos/AxisPositionGizmo.h
--- a/Engine/Source/Runtime/Experimental/InteractiveToolsFramework/Public/BaseGizmos/AxisPositionGizmo.h
+++ b/Engine/Source/Runtime/Experimental/InteractiveToolsFramework/Public/BaseGizmos/AxisPositionGizmo.h
@@ -75,7 +75,7 @@ public:
 	TScriptInterface<IGizmoFloatParameterSource> ParameterSource;
 
 	UPROPERTY()
-	UGizmoViewContext* GizmoViewContext;
+	TObjectPtr<UGizmoViewContext> GizmoViewContext;
 
 	/** The HitTarget provides a hit-test against some 3D element (presumably a visual widget) that controls when interaction can start */
 	UPROPERTY()
diff --git a/Engine/Source/Runtime/Experimental/InteractiveToolsFramework/Public/BaseGizmos/GizmoBaseComponent.h b/Engine/Source/Runtime/Experimental/InteractiveToolsFramework/Public/BaseGizmos/GizmoBaseComponent.h
--- a/Engine/Source/Runtime/Experimental/InteractiveToolsFramework/Public/BaseGizmos/GizmoBaseComponent.h
+++ b/Engine/Source/Runtime/Experimental/InteractiveToolsFramework/Public/BaseGizmos/GizmoBaseComponent.h
@@ -82,7 +82,7 @@ protected:
 	bool bWorld = false;
 
 	UPROPERTY()
-	UGizmoViewContext* GizmoViewContext = nullptr;
+	TObjectPtr<UGizmoViewContext> GizmoViewContext = nullptr;
 
 	// True when GizmoViewContext is not null. Here as a boolean so it
 	// can be pointed to by the proxy to see what it should do.
diff --git a/Engine/Source/Runtime/Interchange/Core/Public/InterchangeFactoryBase.h b/Engine/Source/Runtime/Interchange/Core/Public/InterchangeFactoryBase.h
--- a/Engine/Source/Runtime/Interchange/Core/Public/InterchangeFactoryBase.h
+++ b/Engine/Source/Runtime/Interchange/Core/Public/InterchangeFactoryBase.h
@@ -168,5 +168,5 @@ public:
 
 
 	UPROPERTY()
-	UInterchangeResultsContainer* Results;
+	TObjectPtr<UInterchangeResultsContainer> Results;
 };
diff --git a/Engine/Source/Runtime/Interchange/Core/Public/InterchangeResult.h b/Engine/Source/Runtime/Interchange/Core/Public/InterchangeResult.h
--- a/Engine/Source/Runtime/Interchange/Core/Public/InterchangeResult.h
+++ b/Engine/Source/Runtime/Interchange/Core/Public/InterchangeResult.h
@@ -58,7 +58,7 @@ public:
 	FString DestinationAssetName;
 
 	UPROPERTY()
-	const UClass* AssetType;
+	TObjectPtr<const UClass> AssetType;
 
 	UPROPERTY()
 	FString InterchangeKey;
diff --git a/Engine/Source/Runtime/Interchange/Core/Public/InterchangeResultsContainer.h b/Engine/Source/Runtime/Interchange/Core/Public/InterchangeResultsContainer.h
--- a/Engine/Source/Runtime/Interchange/Core/Public/InterchangeResultsContainer.h
+++ b/Engine/Source/Runtime/Interchange/Core/Public/InterchangeResultsContainer.h
@@ -75,5 +75,5 @@ private:
 	mutable FCriticalSection Lock;
 
 	UPROPERTY()
-	TArray<UInterchangeResult*> Results;
+	TArray<TObjectPtr<UInterchangeResult>> Results;
 };
diff --git a/Engine/Source/Runtime/MeshDescription/Public/MeshDescriptionBaseBulkData.h b/Engine/Source/Runtime/MeshDescription/Public/MeshDescriptionBaseBulkData.h
--- a/Engine/Source/Runtime/MeshDescription/Public/MeshDescriptionBaseBulkData.h
+++ b/Engine/Source/Runtime/MeshDescription/Public/MeshDescriptionBaseBulkData.h
@@ -48,10 +48,10 @@ protected:
 	TPimplPtr<FMeshDescriptionBulkData> BulkData;
 
 	UPROPERTY(Transient, Instanced)
-	UMeshDescriptionBase* PreallocatedMeshDescription;
+	TObjectPtr<UMeshDescriptionBase> PreallocatedMeshDescription;
 
 	UPROPERTY(Transient)
-	UMeshDescriptionBase* MeshDescription;
+	TObjectPtr<UMeshDescriptionBase> MeshDescription;
 
 #endif //WITH_EDITORONLY_DATA
 };
diff --git a/Engine/Source/Runtime/MovieScene/Private/Tests/MovieSceneTestObjects.h b/Engine/Source/Runtime/MovieScene/Private/Tests/MovieSceneTestObjects.h
--- a/Engine/Source/Runtime/MovieScene/Private/Tests/MovieSceneTestObjects.h
+++ b/Engine/Source/Runtime/MovieScene/Private/Tests/MovieSceneTestObjects.h
@@ -124,7 +124,7 @@ public:
 #endif
 
 	UPROPERTY()
-	TArray<UMovieSceneSection*> SectionArray;
+	TArray<TObjectPtr<UMovieSceneSection>> SectionArray;
 };
 
 UCLASS(MinimalAPI)
diff --git a/Engine/Source/Runtime/MovieSceneTracks/Public/Systems/MovieScenePredictionSystem.h b/Engine/Source/Runtime/MovieSceneTracks/Public/Systems/MovieScenePredictionSystem.h
--- a/Engine/Source/Runtime/MovieSceneTracks/Public/Systems/MovieScenePredictionSystem.h
+++ b/Engine/Source/Runtime/MovieSceneTracks/Public/Systems/MovieScenePredictionSystem.h
@@ -154,11 +154,11 @@ private:
 
 	/** The sequence player we're interrogating */
 	UPROPERTY()
-	UMovieSceneSequencePlayer* SequencePlayer;
+	TObjectPtr<UMovieSceneSequencePlayer> SequencePlayer;
 
 	/** The target scene component we're interrogating */
 	UPROPERTY()
-	USceneComponent* SceneComponent;
+	TObjectPtr<USceneComponent> SceneComponent;
 
 	/** The time at which the interrogation should run */
 	FFrameTime RootPredictedTime;
@@ -195,8 +195,8 @@ private:
 	UE::MovieScene::FInterrogationChannels InterrogationChannels;
 
 	UPROPERTY()
-	TArray<UMovieSceneAsyncAction_SequencePrediction*> PendingPredictions;
+	TArray<TObjectPtr<UMovieSceneAsyncAction_SequencePrediction>> PendingPredictions;
 
 	UPROPERTY()
-	TArray<UMovieSceneAsyncAction_SequencePrediction*> ProcessingPredictions;
+	TArray<TObjectPtr<UMovieSceneAsyncAction_SequencePrediction>> ProcessingPredictions;
 };
diff --git a/Engine/Source/Runtime/RigVM/Public/RigVMCore/RigVM.h b/Engine/Source/Runtime/RigVM/Public/RigVMCore/RigVM.h
--- a/Engine/Source/Runtime/RigVM/Public/RigVMCore/RigVM.h
+++ b/Engine/Source/Runtime/RigVM/Public/RigVMCore/RigVM.h
@@ -113,7 +113,7 @@ private:
 	FString CPPType;
 
 	UPROPERTY(transient)
-	UScriptStruct* ScriptStruct;
+	TObjectPtr<UScriptStruct> ScriptStruct;
 
 	UPROPERTY()
 	FName ScriptStructPath;
@@ -682,7 +682,7 @@ private:
 	int32 ExecutingThreadId;
 
 	UPROPERTY(transient)
-	URigVM* DeferredVMToCopy;
+	TObjectPtr<URigVM> DeferredVMToCopy;
 
 	void CopyDeferredVMIfRequired();
 
diff --git a/Engine/Source/Runtime/RigVM/Public/RigVMCore/RigVMMemory.h b/Engine/Source/Runtime/RigVM/Public/RigVMCore/RigVMMemory.h
--- a/Engine/Source/Runtime/RigVM/Public/RigVMCore/RigVMMemory.h
+++ b/Engine/Source/Runtime/RigVM/Public/RigVMCore/RigVMMemory.h
@@ -253,7 +253,7 @@ struct RIGVM_API FRigVMRegister
 	// The resolved UScriptStruct / UObject (in the future)
 	// used for debugging.
 	UPROPERTY(transient)
-	UObject* BaseCPPTypeObject;
+	TObjectPtr<UObject> BaseCPPTypeObject;
 
 #endif
 
@@ -388,11 +388,11 @@ private:
 
 	// The c++ script struct of the resulting memory (for example FVector for Transform.Translation)
 	UPROPERTY()
-	UScriptStruct* ScriptStruct;
+	TObjectPtr<UScriptStruct> ScriptStruct;
 
 	// The c++ script struct of the source memory (for example FTransform for Transform.Translation)
 	UPROPERTY()
-	UScriptStruct* ParentScriptStruct;
+	TObjectPtr<UScriptStruct> ParentScriptStruct;
 
 	// The index of the element within an array (for example 3 for Transform[3])
 	UPROPERTY()
@@ -1465,7 +1465,7 @@ private:
 	TArray<uint8> Data;
 
 	UPROPERTY(transient)
-	TArray<UScriptStruct*> ScriptStructs;
+	TArray<TObjectPtr<UScriptStruct>> ScriptStructs;
 
 	UPROPERTY(transient)
 	TMap<FName, int32> NameMap;
diff --git a/Engine/Source/Runtime/Serialization/Private/Tests/StructSerializerTestTypes.h b/Engine/Source/Runtime/Serialization/Private/Tests/StructSerializerTestTypes.h
--- a/Engine/Source/Runtime/Serialization/Private/Tests/StructSerializerTestTypes.h
+++ b/Engine/Source/Runtime/Serialization/Private/Tests/StructSerializerTestTypes.h
@@ -131,7 +131,7 @@ struct FStructSerializerObjectTestStruct
 	GENERATED_BODY()
 
 	UPROPERTY()
-	UClass* RawClass;
+	TObjectPtr<UClass> RawClass;
 
 	UPROPERTY()
 	TObjectPtr<UClass> WrappedClass;
@@ -143,7 +143,7 @@ struct FStructSerializerObjectTestStruct
 	TSoftClassPtr<class UMetaData> SoftClass;
 
 	UPROPERTY()
-	class UObject* RawObject;
+	TObjectPtr<class UObject> RawObject;
 
 	UPROPERTY()
 	TObjectPtr<class UObject> WrappedObject;
