Only try to find a scene component target to attach to if the pasted component is a scene component. Non-Scene components can only be attached to the root actor. This solves an issue where a pasted non-scene component could be incorrectly attached to the default scene root.

#rnx

diff --git a/Engine/Source/Editor/SubobjectDataInterface/Private/SubobjectDataSubsystem.cpp b/Engine/Source/Editor/SubobjectDataInterface/Private/SubobjectDataSubsystem.cpp
--- a/Engine/Source/Editor/SubobjectDataInterface/Private/SubobjectDataSubsystem.cpp
+++ b/Engine/Source/Editor/SubobjectDataInterface/Private/SubobjectDataSubsystem.cpp
@@ -2030,19 +2030,27 @@ void USubobjectDataSubsystem::PasteSubobjects(const FSubobjectDataHandle& PasteT
 				NewActorComponent = nullptr;
 			}
 
-			// The default target to attach the new pasted component to is the Scene Root for the context
-			FSubobjectDataHandle TargetParentHandle = FindSceneRootForSubobject(FindParentForNewSubobject(NewActorComponent, PasteToContext));
+			// Find a suitable parent handle to be used as the default, which would be the root actor
+			FSubobjectDataHandle TargetParentHandle = FindParentForNewSubobject(NewActorComponent, PasteToContext);
 
-			// If there were any handles that we should target the attachment for instead, use that
-			for (FSubobjectDataHandle SelectedNode : NewParentHandles)
+			// Ensure that for scene components they are attached to the scene root. Non-Scene components (i.e. a blackboard)
+			// should stay attached to the root actor, not the scene root.
+			if (NewActorComponent && NewActorComponent->IsA<USceneComponent>())
 			{
-				const FSubobjectData* SelectedData = SelectedNode.GetSharedDataPtr().Get();
+				// A scene component should be attaching to the scene root by default.
+				TargetParentHandle = FindSceneRootForSubobject(TargetParentHandle);
 
-				// Only scene components can be attached to when pasting
-				if (SelectedData && SelectedData->IsSceneComponent())
+				// If there were any handles that should be considered instead, see if they are valid.
+				for (FSubobjectDataHandle SelectedNode : NewParentHandles)
 				{
-					TargetParentHandle = SelectedData->GetHandle();
-					break;
+					const FSubobjectData* SelectedData = SelectedNode.GetSharedDataPtr().Get();
+
+					// Only scene components can be attached to when pasting
+					if (SelectedData && SelectedData->IsSceneComponent())
+					{
+						TargetParentHandle = SelectedData->GetHandle();
+						break;
+					}
 				}
 			}
 
