Change the behavior when pasting subobjects to attach to the selected subobject's parent instead of the subobject itself. This makes the pasting behavior between the blueprint and level instance editor the same.

#rnx

diff --git a/Engine/Source/Editor/SubobjectDataInterface/Private/SubobjectDataSubsystem.cpp b/Engine/Source/Editor/SubobjectDataInterface/Private/SubobjectDataSubsystem.cpp
--- a/Engine/Source/Editor/SubobjectDataInterface/Private/SubobjectDataSubsystem.cpp
+++ b/Engine/Source/Editor/SubobjectDataInterface/Private/SubobjectDataSubsystem.cpp
@@ -2046,9 +2046,9 @@ void USubobjectDataSubsystem::PasteSubobjects(const FSubobjectDataHandle& PasteT
 					const FSubobjectData* SelectedData = SelectedNode.GetSharedDataPtr().Get();
 
 					// Only scene components can be attached to when pasting
-					if (SelectedData && SelectedData->IsSceneComponent())
+					if (SelectedData && SelectedData->IsSceneComponent() && SelectedData->HasParent())
 					{
-						TargetParentHandle = SelectedData->GetHandle();
+						TargetParentHandle = SelectedData->GetParentHandle();
 						break;
 					}
 				}
