When pasting a subobject, check to see if the desired parent is a scene component before allowing it to be the target of the paste. This prevents the incorrect parenting to the Actor root in the case of copy/pasting the default scene root.

#rnx

diff --git a/Engine/Source/Editor/SubobjectDataInterface/Private/SubobjectDataSubsystem.cpp b/Engine/Source/Editor/SubobjectDataInterface/Private/SubobjectDataSubsystem.cpp
--- a/Engine/Source/Editor/SubobjectDataInterface/Private/SubobjectDataSubsystem.cpp
+++ b/Engine/Source/Editor/SubobjectDataInterface/Private/SubobjectDataSubsystem.cpp
@@ -2048,8 +2048,13 @@ void USubobjectDataSubsystem::PasteSubobjects(const FSubobjectDataHandle& PasteT
 					// Only scene components can be attached to when pasting
 					if (SelectedData && SelectedData->IsSceneComponent() && SelectedData->HasParent())
 					{
-						TargetParentHandle = SelectedData->GetParentHandle();
-						break;
+						const FSubobjectData* ParentData = SelectedData->GetParentHandle().GetData();
+						// The parent data can only be used if it's a scene component (not an actor or non-scene component)
+						if (ParentData && ParentData->IsSceneComponent())
+						{
+							TargetParentHandle = SelectedData->GetParentHandle();
+							break;
+						}
 					}
 				}
 			}
