Do not allow copying or duplicating subobjects with the name of "DefaultSceneRoot". This will prevent confusing renaming situations where the DefaultSceneRoot could be deleted, but still be preventing other subobjects from being named that

#rnx

diff --git a/Engine/Source/Editor/SubobjectDataInterface/Private/SubobjectData.cpp b/Engine/Source/Editor/SubobjectDataInterface/Private/SubobjectData.cpp
--- a/Engine/Source/Editor/SubobjectDataInterface/Private/SubobjectData.cpp
+++ b/Engine/Source/Editor/SubobjectDataInterface/Private/SubobjectData.cpp
@@ -74,6 +74,11 @@ bool FSubobjectData::CanDuplicate() const
 
 bool FSubobjectData::CanCopy() const
 {
+	if (IsInstancedInheritedComponent())
+	{
+		return FComponentEditorUtils::CanCopyComponent(GetComponentTemplate()) && !SceneRootHasDefaultName();
+	}
+	
 	return FComponentEditorUtils::CanCopyComponent(GetComponentTemplate());
 }
 
@@ -110,7 +115,7 @@ bool FSubobjectData::CanRename() const
 	// is on an instance in the level
 	if(GetSCSNode() != nullptr && !IsInheritedSCSNode())
 	{
-		return !IsInstancedInheritedComponent();
+		return !IsInstancedInheritedComponent() && !SceneRootHasDefaultName();
 	}
 	
 	return !IsInheritedComponent() && !IsDefaultSceneRoot() && !IsChildActorSubtreeObject();
@@ -928,6 +933,20 @@ bool FSubobjectData::IsDefaultSceneRoot() const
 	return false;
 }
 
+bool FSubobjectData::SceneRootHasDefaultName() const
+{
+	const UActorComponent* Template = GetComponentTemplate();
+	const FName TemplateName = Template ? Template->GetFName() : NAME_None;
+
+	// Only the first default scene root can be attached to an actor, so this will 
+	// rule out false positives of other subobjects that may have been named the same thing
+	// or if the first default scene root was duplicated
+	const FSubobjectData* ParentData = GetParentHandle().GetData();
+	const bool bIsAttachedToActor = ParentData ? ParentData->IsActor() : false;
+
+	return bIsAttachedToActor && TemplateName.ToString().StartsWith(USceneComponent::GetDefaultSceneRootVariableName().ToString());
+}
+
 bool FSubobjectData::IsComponent() const
 {
 	// Check if we are pointing to a component
diff --git a/Engine/Source/Editor/SubobjectDataInterface/Public/SubobjectData.h b/Engine/Source/Editor/SubobjectDataInterface/Public/SubobjectData.h
--- a/Engine/Source/Editor/SubobjectDataInterface/Public/SubobjectData.h
+++ b/Engine/Source/Editor/SubobjectDataInterface/Public/SubobjectData.h
@@ -126,6 +126,8 @@ public:
 	virtual bool IsSceneComponent() const;
 	virtual bool IsRootComponent() const;
 	virtual bool IsDefaultSceneRoot() const;
+	/** Returns true if the component template's FName starts with USceneComponent::GetDefaultSceneRootVariableName */
+	bool SceneRootHasDefaultName() const;
 	/* Returns true if this subobject is a component. */
 	virtual bool IsComponent() const;
 	virtual bool IsChildActor() const;
