Set the parent of non-scene components correctly when converting an instance to a blueprint

#rnx
#p4v-preflight-copy 17105167

diff --git a/Engine/Source/Editor/SubobjectDataInterface/Private/SubobjectDataSubsystem.cpp b/Engine/Source/Editor/SubobjectDataInterface/Private/SubobjectDataSubsystem.cpp
--- a/Engine/Source/Editor/SubobjectDataInterface/Private/SubobjectDataSubsystem.cpp
+++ b/Engine/Source/Editor/SubobjectDataInterface/Private/SubobjectDataSubsystem.cpp
@@ -301,7 +301,13 @@ void USubobjectDataSubsystem::GatherSubobjectData(UObject* Context, TArray<FSubo
 					// Otherwise add a node parents to the root actor
 					else
 					{
-						NewHandle = FactoryCreateInheritedBpSubobject(SCS_Node, RootComponentHandle.IsValid() ? RootComponentHandle : RootActorHandle, /* bIsInherited = */ StackIndex > 0, OutArray);
+						FSubobjectDataHandle& ParentHandle = RootActorHandle;
+						// Only a valid scene component can be used as the parent handle for a subobject, otherwise it should be attached to the root actor
+						if (RootComponentHandle.IsValid() && SCS_Node->ComponentTemplate->IsA<USceneComponent>())
+						{
+							ParentHandle = RootComponentHandle;
+						}
+						NewHandle = FactoryCreateInheritedBpSubobject(SCS_Node, ParentHandle, /* bIsInherited = */ StackIndex > 0, OutArray);
 					}
 					NewData = NewHandle.GetData();
 					
