[Subobject Editor] Use the correct previous subobject transform when conforming a subobject to it's new parent during reparenting. Previously it was incorrectly using the same "Dropped On" subobject (which would be the new parent) thus always giving it a relative position/rotation/scale of 0 when conforming.

#rnx

diff --git a/Engine/Source/Editor/SubobjectDataInterface/Private/SubobjectDataSubsystem.cpp b/Engine/Source/Editor/SubobjectDataInterface/Private/SubobjectDataSubsystem.cpp
--- a/Engine/Source/Editor/SubobjectDataInterface/Private/SubobjectDataSubsystem.cpp
+++ b/Engine/Source/Editor/SubobjectDataInterface/Private/SubobjectDataSubsystem.cpp
@@ -1603,7 +1603,7 @@ bool USubobjectDataSubsystem::ReparentSubobjects(const FReparentSubobjectParams&
 				AttachSubobject(NewParentData->GetHandle(), DroppedData->GetHandle());
 
 				// Attempt to locate a matching instance of the parent component template in the Actor context that's being edited
-				USceneComponent* ParentSceneComponent = Cast<USceneComponent>(DroppedData->FindMutableComponentInstanceInActor(Params.ActorPreviewContext));
+				USceneComponent* ParentSceneComponent = OldParentData ? Cast<USceneComponent>(OldParentData->FindMutableComponentInstanceInActor(Params.ActorPreviewContext)) : nullptr;
 				if(SceneComponentTemplate && ParentSceneComponent && ParentSceneComponent->IsRegistered())
 				{
 					ConformTransformRelativeToParent(SceneComponentTemplate, ParentSceneComponent);
