[Subobject Editor] Allow deletion of subobjects that are the default scene root so that they get properly replaced with a "DefaultSceneRoot" in the subobject blueprint editor.

#rnx

diff --git a/Engine/Source/Editor/SubobjectDataInterface/Private/InheritedSubobjectData.cpp b/Engine/Source/Editor/SubobjectDataInterface/Private/InheritedSubobjectData.cpp
--- a/Engine/Source/Editor/SubobjectDataInterface/Private/InheritedSubobjectData.cpp
+++ b/Engine/Source/Editor/SubobjectDataInterface/Private/InheritedSubobjectData.cpp
@@ -44,7 +44,7 @@ bool FInheritedSubobjectData::CanEdit() const
 
 bool FInheritedSubobjectData::CanDelete() const
 {
-	if(IsInheritedComponent() || IsDefaultSceneRoot() || (GetSCSNode() != nullptr && IsInstancedInheritedComponent()) || IsChildActorSubtreeObject())
+	if(IsInheritedComponent() || (IsDefaultSceneRoot() && SceneRootHasDefaultName()) || (GetSCSNode() != nullptr && IsInstancedInheritedComponent()) || IsChildActorSubtreeObject())
 	{
 		return false;
 	}
