Fix an ensure that was caused when adding a runtime subobject to an editor only subobject.

#rnx

diff --git a/Engine/Source/Editor/SubobjectDataInterface/Private/SubobjectDataSubsystem.cpp b/Engine/Source/Editor/SubobjectDataInterface/Private/SubobjectDataSubsystem.cpp
--- a/Engine/Source/Editor/SubobjectDataInterface/Private/SubobjectDataSubsystem.cpp
+++ b/Engine/Source/Editor/SubobjectDataInterface/Private/SubobjectDataSubsystem.cpp
@@ -2352,7 +2352,9 @@ FSubobjectDataHandle USubobjectDataSubsystem::FactoryCreateInheritedBpSubobject(
 
 			DetachSubobject(OldParentPtr->GetHandle(), NewData->GetHandle());
 			NodeSCS->RemoveNode(OldParentPtr->GetSCSNode());
-			
+			const bool bWasRemoved = OutArray.Remove(OldParentPtr->GetHandle()) > 0;
+			ensure(bWasRemoved);
+
 			// if the grandparent node is invalid (assuming this means that the parent node was the scene-root)
 			if (!GrandparentPtr->IsValid())
 			{
@@ -2382,8 +2384,11 @@ FSubobjectDataHandle USubobjectDataSubsystem::FactoryCreateInheritedBpSubobject(
 	{
 		OutArray.Add(OutHandle);
 
-		for (USCS_Node* ChildNode : InSCSNode->GetChildNodes())
+		// Note: The child array may be modified if a runtime subobject is parented to an editor only subobject
+		const TArray<USCS_Node*>& ChildNodes = InSCSNode->GetChildNodes();
+		for (int32 i = ChildNodes.Num() - 1; i >= 0; --i)
 		{
+			USCS_Node* ChildNode = ChildNodes[i];
 			FSubobjectDataHandle NewChildHandle = FactoryCreateInheritedBpSubobject(ChildNode, OutHandle, bIsInherited, OutArray);
 			ensure(NewChildHandle.IsValid());
 			OutArray.Add(NewChildHandle);
