Fix crashing when undoing a native component class change in a blueprint
Fix asserts as a result of replacing actor components without going through proper register/unregister lifecycle
Move logic for component class changing out of the details customization and in to the subobject data subsystem

diff --git a/Engine/Source/Editor/Kismet/Private/BlueprintDetailsCustomization.cpp b/Engine/Source/Editor/Kismet/Private/BlueprintDetailsCustomization.cpp
--- a/Engine/Source/Editor/Kismet/Private/BlueprintDetailsCustomization.cpp
+++ b/Engine/Source/Editor/Kismet/Private/BlueprintDetailsCustomization.cpp
@@ -6695,49 +6695,10 @@ void FBlueprintComponentDetails::HandleNewEntryClassSelected(const UClass* NewEn
 
 		if (const UActorComponent* ComponentTemplate = CachedNodePtr->GetComponentTemplate())
 		{
-			const FScopedTransaction Transaction(LOCTEXT("SetComponentClassOverride", "Set Component Class Override"));
-
-			const FName ComponentTemplateName = ComponentTemplate->GetFName();
-
-			UBlueprint* BlueprintObj = GetBlueprintObj();
-			check(BlueprintObj);
-
-			BlueprintObj->Modify();
-
-			if (FBPComponentClassOverride* Override = BlueprintObj->ComponentClassOverrides.FindByKey(ComponentTemplateName))
-			{
-				bool bRemoveEntry = false;
-				bool bFoundOverride = false;
-
-				UBlueprint* ParentBP = UBlueprint::GetBlueprintFromClass(Cast<UBlueprintGeneratedClass>(BlueprintObj->ParentClass));
-				while (ParentBP)
-				{
-					if (FBPComponentClassOverride* ParentOverride = BlueprintObj->ComponentClassOverrides.FindByKey(ComponentTemplateName))
-					{
-						bRemoveEntry = (ParentOverride->ComponentClass == NewEntryClass);
-						bFoundOverride = true;
-						break;
-					}
-				}
-				if (!bFoundOverride)
-				{
-					bRemoveEntry = (ComponentTemplate->GetClass() == NewEntryClass);
-				}
-				if (bRemoveEntry)
-				{
-					BlueprintObj->ComponentClassOverrides.RemoveAllSwap([ComponentTemplateName](const FBPComponentClassOverride& CCOverride) { return (CCOverride.ComponentName == ComponentTemplateName); });
-				}
-				else
-				{
-					Override->ComponentClass = const_cast<UClass*>(NewEntryClass);
-				}
-			}
-			else
+			if (USubobjectDataSubsystem* System = USubobjectDataSubsystem::Get())
 			{
-				BlueprintObj->ComponentClassOverrides.Emplace(FBPComponentClassOverride(ComponentTemplateName, const_cast<UClass*>(NewEntryClass)));
+				System->ChangeSubobjectClass(CachedNodePtr->GetDataHandle(), NewEntryClass);
 			}
-
-			FBlueprintEditorUtils::MarkBlueprintAsStructurallyModified(BlueprintObj);
 		}
 	}
 }
diff --git a/Engine/Source/Editor/SubobjectDataInterface/Private/SubobjectDataSubsystem.cpp b/Engine/Source/Editor/SubobjectDataInterface/Private/SubobjectDataSubsystem.cpp
--- a/Engine/Source/Editor/SubobjectDataInterface/Private/SubobjectDataSubsystem.cpp
+++ b/Engine/Source/Editor/SubobjectDataInterface/Private/SubobjectDataSubsystem.cpp
@@ -14,6 +14,7 @@
 #include "BaseSubobjectDataFactory.h"
 #include "ChildSubobjectDataFactory.h"
 #include "InheritedSubobjectDataFactory.h"
+#include "Serialization/ArchiveReplaceOrClearExternalReferences.h"
 
 #include "BlueprintEditorSettings.h"	// bHideConstructionScriptComponentsInDetailsView
 #include "Kismet2/ComponentEditorUtils.h"
@@ -1334,6 +1335,256 @@ bool USubobjectDataSubsystem::RenameSubobject(const FSubobjectDataHandle& Handle
 	return false;
 }
 
+bool USubobjectDataSubsystem::ChangeSubobjectClass(const FSubobjectDataHandle& Handle, const UClass* NewClass)
+{
+	if (!GetAllowNativeComponentClassOverrides())
+	{
+		return false;
+	}
+
+	const FSubobjectData* Data = Handle.GetData();
+	if (Data && Data->IsNativeComponent())
+	{
+		// For instanced components
+		if (UActorComponent* ComponentTemplate = Data->GetMutableComponentTemplate())
+		{
+			if (UBlueprint* BlueprintObj = Data->GetBlueprint())
+			{
+				UClass* BaseClass = ComponentTemplate->GetClass();
+
+				if (const FBPComponentClassOverride* Override = BlueprintObj->ComponentClassOverrides.FindByKey(ComponentTemplate->GetFName()))
+				{
+					AActor* Owner = ComponentTemplate->GetOwner();
+					AActor* OwnerArchetype = CastChecked<AActor>(Owner->GetArchetype());
+					if (UActorComponent* ArchetypeComponent = Cast<UActorComponent>((UObject*)FindObjectWithOuter(OwnerArchetype, UActorComponent::StaticClass(), ComponentTemplate->GetFName())))
+					{
+						BaseClass = ArchetypeComponent->GetClass();
+					}
+				}
+
+				if (!NewClass->IsChildOf(BaseClass))
+				{
+					return false;
+				}
+
+				const FScopedTransaction Transaction(LOCTEXT("SetComponentClassOverride", "Set Component Class Override"));
+
+				const FName ComponentTemplateName = ComponentTemplate->GetFName();
+
+				BlueprintObj->Modify();
+
+				if (FBPComponentClassOverride* Override = BlueprintObj->ComponentClassOverrides.FindByKey(ComponentTemplateName))
+				{
+					bool bRemoveEntry = false;
+					bool bFoundOverride = false;
+
+					UBlueprint* ParentBP = UBlueprint::GetBlueprintFromClass(Cast<UBlueprintGeneratedClass>(BlueprintObj->ParentClass));
+					while (ParentBP)
+					{
+						if (FBPComponentClassOverride* ParentOverride = BlueprintObj->ComponentClassOverrides.FindByKey(ComponentTemplateName))
+						{
+							bRemoveEntry = (ParentOverride->ComponentClass == NewClass);
+							bFoundOverride = true;
+							break;
+						}
+					}
+					if (!bFoundOverride)
+					{
+						bRemoveEntry = (ComponentTemplate->GetClass() == NewClass);
+					}
+					if (bRemoveEntry)
+					{
+						BlueprintObj->ComponentClassOverrides.RemoveAllSwap([ComponentTemplateName](const FBPComponentClassOverride& CCOverride) { return (CCOverride.ComponentName == ComponentTemplateName); });
+					}
+					else
+					{
+						Override->ComponentClass = NewClass;
+					}
+				}
+				else
+				{
+					BlueprintObj->ComponentClassOverrides.Emplace(FBPComponentClassOverride(ComponentTemplateName, NewClass));
+				}
+
+				// Custom transaction change that operates on the UBlueprint and replaces all instances of the subobject with one of the new class
+				struct FSubobjectClassChange : public FCommandChange
+				{
+					FSubobjectClassChange(FGuid InOperationGuid, FName InTemplateName, const UClass* InApplyClass, const UClass* InRevertClass)
+						: OperationGuid(InOperationGuid)
+						, TemplateName(InTemplateName)
+						, ApplyClass(InApplyClass)
+						, RevertClass(InRevertClass)
+					{
+					}
+
+					// Utility function to assemble the name of the object in the transient package to use as template from replacement
+					static FString MakeReplacementName(const FString& CurrentName, UObject* SubobjectOuter, const TCHAR* ReplaceName, FGuid OperationGuid)
+					{
+						return FString::Printf(TEXT("%s_%s_%s_%d"), *CurrentName, ReplaceName, *OperationGuid.ToString(), SubobjectOuter->GetUniqueID());
+					}
+
+					// Swaps out all of the supplied subobject instances with objects of the new class
+					void SwapObjects(TArrayView<UObject*> SubobjectInstances, const TCHAR* ReplaceName, const UClass* NewClass, TArrayView<UObject*> ReplacementTemplates = TArrayView<UObject*>())
+					{					
+						TMap<UObject*, UObject*> ReferenceReplacementMap;
+						const FString TemplateNameStr = TemplateName.ToString();
+
+						for (int32 ObjectIndex = 0; ObjectIndex < SubobjectInstances.Num(); ++ObjectIndex)
+						{
+							UObject* Subobject = SubobjectInstances[ObjectIndex];
+							UObject* SubobjectOuter = Subobject->GetOuter();
+							const EObjectFlags SubobjectFlags = Subobject->GetFlags();
+
+							bool bWasRegistered = false;
+							if (UActorComponent* Comp = Cast<UActorComponent>(Subobject))
+							{
+								bWasRegistered = Comp->IsRegistered();
+								if (bWasRegistered)
+								{
+									Comp->UnregisterComponent();
+								}
+							}
+
+							Subobject->Rename(nullptr, GetTransientPackage(), REN_DontCreateRedirectors | REN_ForceNoResetLoaders | REN_NonTransactional);
+							Subobject->MarkAsGarbage();
+
+							UObject* NewInstance = nullptr;
+							UObject* ReplaceTemplate = nullptr;
+							if (ReplacementTemplates.Num() > 0)
+							{
+								ReplaceTemplate = ReplacementTemplates[ObjectIndex];
+							}
+							else
+							{
+								ReplaceTemplate = static_cast<UObject*>(FindObjectWithOuter(GetTransientPackage(), nullptr, *MakeReplacementName(TemplateNameStr, SubobjectOuter, ReplaceName, OperationGuid)));
+							}
+
+							if (ReplaceTemplate)
+							{ 
+								NewInstance = StaticDuplicateObject(ReplaceTemplate, SubobjectOuter, TemplateName);
+							}
+							else
+							{
+								// Fallback if our replacement template was lost
+								NewInstance = NewObject<UObject>(SubobjectOuter, NewClass, TemplateName, SubobjectFlags);
+								UEngine::CopyPropertiesForUnrelatedObjects(Subobject, NewInstance);
+							}
+
+							if (bWasRegistered)
+							{
+								CastChecked<UActorComponent>(NewInstance)->RegisterComponent();
+							}
+
+							ReferenceReplacementMap.Add(Subobject, NewInstance);
+
+							FArchiveReplaceOrClearExternalReferences<UObject> ReplaceAr(SubobjectOuter, ReferenceReplacementMap, SubobjectOuter->GetOutermost());
+						}
+
+						GEngine->NotifyToolsOfObjectReplacement(ReferenceReplacementMap);
+					}
+
+					void SwapObjects(UObject* Object, const bool bApply)
+					{
+						const TCHAR* ReplaceName = (bApply ? TEXT("REPLACEMENT") : TEXT("REPLACED"));
+						const UClass* CurrentClass = (bApply ? ApplyClass : RevertClass);
+						const UClass* NewClass = (bApply ? RevertClass : ApplyClass);
+
+						UBlueprint* BP = CastChecked<UBlueprint>(Object);
+						UObject* Template = static_cast<UObject*>(FindObjectWithOuter(BP->GeneratedClass->GetDefaultObject(), nullptr, TemplateName));
+						ensure(Template && Template->GetClass() == CurrentClass);
+
+						TArray<UObject*> ArchetypeInstances;
+						Template->GetArchetypeInstances(ArchetypeInstances);
+						ArchetypeInstances.Add(Template);
+
+						SwapObjects(ArchetypeInstances, ReplaceName, NewClass);
+
+						FBlueprintEditorUtils::MarkBlueprintAsStructurallyModified(BP);
+					}
+
+					virtual void Apply(UObject* Object) override
+					{
+						SwapObjects(Object, true);
+					}
+
+					virtual void Revert(UObject* Object) override
+					{
+						SwapObjects(Object, false);
+					}
+
+					virtual void AddReferencedObjects(FReferenceCollector& Collector) override
+					{
+						Collector.AddReferencedObject(ApplyClass);
+						Collector.AddReferencedObject(RevertClass);
+					}
+
+					virtual FString ToString() const override
+					{
+						return FString();
+					}
+
+					FGuid OperationGuid;
+					FName TemplateName;
+					const UClass* ApplyClass;
+					const UClass* RevertClass;
+				};
+
+				// To avoid all of the subobjects we are manipulating getting serialized in to the transaction buffer unnecessarily
+				// we record them with a simple custom dummy custom change as we don't actually need them to be transacted since
+				// the primary custom change operation on the blueprint will handle destroying and replacing them
+				struct FDummySubobjectClassChange : FCommandChange
+				{
+					virtual void Apply(UObject* Object) override {}
+					virtual void Revert(UObject* Object) override {}
+					virtual FString ToString() const override { return FString(); }
+				};
+
+				const FGuid OperationGuid = FGuid::NewGuid();
+				const FString CurrentName = ComponentTemplate->GetName();
+
+				TArray<UObject*> ArchetypeInstances;
+				ComponentTemplate->GetArchetypeInstances(ArchetypeInstances);
+				ArchetypeInstances.Add(ComponentTemplate);
+
+				TArray<UObject*> ReplacementObjects;
+				ReplacementObjects.Reserve(ArchetypeInstances.Num());
+
+				for (UObject* ArchetypeInstance : ArchetypeInstances)
+				{
+					GUndo->StoreUndo(ArchetypeInstance, MakeUnique<FDummySubobjectClassChange>());
+
+					UObject* SubobjectOuter = ArchetypeInstance->GetOuter();
+					SubobjectOuter->Modify();
+
+					{
+						const FString ReplacementName = FSubobjectClassChange::MakeReplacementName(CurrentName, SubobjectOuter, TEXT("REPLACEMENT"), OperationGuid);
+						const FString ReplacedName = FSubobjectClassChange::MakeReplacementName(CurrentName, SubobjectOuter, TEXT("REPLACED"), OperationGuid);
+
+						// Avoid the new instance going in to the transaction buffer as we are managing that manually through the custom change record
+						TGuardValue<ITransaction*> SuppressTransaction(GUndo, nullptr);
+						UObject* NewInstance = NewObject<UObject>(GetTransientPackage(), NewClass, *ReplacementName, ArchetypeInstance->GetFlags());
+						UEngine::CopyPropertiesForUnrelatedObjects(ArchetypeInstance, NewInstance);
+						ReplacementObjects.Add(NewInstance);
+
+						StaticDuplicateObject(ArchetypeInstance, GetTransientPackage(), *ReplacedName);
+					}
+				}
+
+				TUniquePtr<FSubobjectClassChange> SubobjectClassChange = MakeUnique<FSubobjectClassChange>(OperationGuid, ComponentTemplate->GetFName(), ComponentTemplate->GetClass(), NewClass);
+				SubobjectClassChange->SwapObjects(ArchetypeInstances, TEXT("REPLACEMENT"), NewClass, ReplacementObjects);
+
+				FBlueprintEditorUtils::MarkBlueprintAsStructurallyModified(BlueprintObj);
+
+				GUndo->StoreUndo(BlueprintObj, MoveTemp(SubobjectClassChange));
+
+				return true;
+			}
+		}
+	}
+
+	return false;
+}
+
 bool USubobjectDataSubsystem::ReparentSubobject(const FReparentSubobjectParams& Params, const FSubobjectDataHandle& ToReparentHandle)
 {
 	TArray<FSubobjectDataHandle> Handles;
diff --git a/Engine/Source/Editor/SubobjectDataInterface/Public/SubobjectDataSubsystem.h b/Engine/Source/Editor/SubobjectDataInterface/Public/SubobjectDataSubsystem.h
--- a/Engine/Source/Editor/SubobjectDataInterface/Public/SubobjectDataSubsystem.h
+++ b/Engine/Source/Editor/SubobjectDataInterface/Public/SubobjectDataSubsystem.h
@@ -249,13 +249,24 @@ public:
 	* Attempts to rename the given subobject to the new name.
 	*
 	* @param Handle			Handle to the subobject to rename
-	* @param InNewName		The new name that is desired for the give subobject
+	* @param InNewName		The new name that is desired for the given subobject
 	*
 	* @return True if the rename was successful, false otherwise.
 	*/
 	UFUNCTION(BlueprintCallable, Category = "SubobjectDataSubsystem")
 	bool RenameSubobject(const FSubobjectDataHandle& Handle, const FText& InNewName);
 
+	/**
+	* Attempts to change the subclass of a native component
+	*
+	* @param Handle			Handle to the subobject to change class of
+	* @param NewClass		The new class that is desired for the given subobject
+	*
+	* @return True if the class change was successful, false otherwise.
+	*/
+	UFUNCTION(BlueprintCallable, Category = "SubobjectDataSubsystem")
+	bool ChangeSubobjectClass(const FSubobjectDataHandle& Handle, const UClass* NewClass);
+
 	/**
 	* Attempts to reparent the given subobject to the new parent
 	*
