Keep absolute location when reparenting a subobject in blueprints

#rnx

diff --git a/Engine/Source/Editor/SubobjectDataInterface/Private/SubobjectDataSubsystem.cpp b/Engine/Source/Editor/SubobjectDataInterface/Private/SubobjectDataSubsystem.cpp
--- a/Engine/Source/Editor/SubobjectDataInterface/Private/SubobjectDataSubsystem.cpp
+++ b/Engine/Source/Editor/SubobjectDataInterface/Private/SubobjectDataSubsystem.cpp
@@ -1857,7 +1857,7 @@ bool USubobjectDataSubsystem::ReparentSubobjects(const FReparentSubobjectParams&
 				AttachSubobject(NewParentData->GetHandle(), DroppedData->GetHandle());
 
 				// Attempt to locate a matching instance of the parent component template in the Actor context that's being edited
-				USceneComponent* ParentSceneComponent = OldParentData ? Cast<USceneComponent>(OldParentData->FindMutableComponentInstanceInActor(Params.ActorPreviewContext)) : nullptr;
+				USceneComponent* ParentSceneComponent = NewParentData ? Cast<USceneComponent>(NewParentData->FindMutableComponentInstanceInActor(Params.ActorPreviewContext)) : nullptr;
 				if(SceneComponentTemplate && ParentSceneComponent && ParentSceneComponent->IsRegistered())
 				{
 					ConformTransformRelativeToParent(SceneComponentTemplate, ParentSceneComponent);
