Fix static analysis warning in subobject subsystem

#rnx

diff --git a/Engine/Source/Editor/SubobjectDataInterface/Private/SubobjectDataSubsystem.cpp b/Engine/Source/Editor/SubobjectDataInterface/Private/SubobjectDataSubsystem.cpp
--- a/Engine/Source/Editor/SubobjectDataInterface/Private/SubobjectDataSubsystem.cpp
+++ b/Engine/Source/Editor/SubobjectDataInterface/Private/SubobjectDataSubsystem.cpp
@@ -1775,7 +1775,11 @@ bool USubobjectDataSubsystem::ReparentSubobjects(const FReparentSubobjectParams&
 	}
 	
 	FSubobjectData* NewParentData = Params.NewParentHandle.GetData();
-	check(NewParentData);
+
+	if (!ensureMsgf(NewParentData, TEXT("There was no valid parent data given! Exiting...")))
+	{
+		return false;
+	}
 	
 	if (Params.BlueprintContext)
 	{
