Add WithIdenticalViaEquality type-trait for FSubobjectDataHandle

This ensures that it compares correctly when used via a UScriptStruct, eg, within Python.

#rnx

diff --git a/Engine/Source/Editor/SubobjectDataInterface/Public/SubobjectDataHandle.h b/Engine/Source/Editor/SubobjectDataInterface/Public/SubobjectDataHandle.h
--- a/Engine/Source/Editor/SubobjectDataInterface/Public/SubobjectDataHandle.h
+++ b/Engine/Source/Editor/SubobjectDataInterface/Public/SubobjectDataHandle.h
@@ -71,4 +71,13 @@ struct SUBOBJECTDATAINTERFACE_API FSubobjectDataHandle
 private:
 	/** Pointer to the actual subobject data that this handle represents */
 	TSharedPtr<FSubobjectData> DataPtr;
-};
\ No newline at end of file
+};
+
+template<>
+struct TStructOpsTypeTraits<FSubobjectDataHandle> : public TStructOpsTypeTraitsBase2<FSubobjectDataHandle>
+{
+	enum
+	{
+		WithIdenticalViaEquality = true,
+	};
+};
