Expose the GetObjectForBlueprint function in the SubobjectDataBluepriuntFunctionLibrary. This allows you to get inherited subobject data for a blueprint.

#rnx

diff --git a/Engine/Source/Editor/SubobjectDataInterface/Public/SubobjectDataBlueprintFunctionLibrary.h b/Engine/Source/Editor/SubobjectDataInterface/Public/SubobjectDataBlueprintFunctionLibrary.h
--- a/Engine/Source/Editor/SubobjectDataInterface/Public/SubobjectDataBlueprintFunctionLibrary.h
+++ b/Engine/Source/Editor/SubobjectDataInterface/Public/SubobjectDataBlueprintFunctionLibrary.h
@@ -87,6 +87,9 @@ public:
 	UFUNCTION(BlueprintCallable, BlueprintPure, Category="Subobject Data")
 	static const UObject* GetObject(const FSubobjectData& Data, bool bEvenIfPendingKill = false) { return Data.GetObject(bEvenIfPendingKill); }
 
+	UFUNCTION(BlueprintCallable, Category="Subobject Data")
+	static const UObject* GetObjectForBlueprint(const FSubobjectData& Data, UBlueprint* Blueprint) { return Data.GetObjectForBlueprint(Blueprint); }
+
 	UFUNCTION(BlueprintCallable, BlueprintPure, Category="Subobject Data")
 	UBlueprint* GetBlueprint(const FSubobjectData& Data) { return Data.GetBlueprint(); }
 
