Ensure that Blueprint subobject (component) names conform to UObject conventions on entry.

diff --git a/Engine/Source/Editor/SubobjectDataInterface/Private/SubobjectDataSubsystem.cpp b/Engine/Source/Editor/SubobjectDataInterface/Private/SubobjectDataSubsystem.cpp
--- a/Engine/Source/Editor/SubobjectDataInterface/Private/SubobjectDataSubsystem.cpp
+++ b/Engine/Source/Editor/SubobjectDataInterface/Private/SubobjectDataSubsystem.cpp
@@ -2086,6 +2086,12 @@ bool USubobjectDataSubsystem::IsValidRename(const FSubobjectDataHandle& Handle,
 			return true;
 		}
 
+		// Subobject names must conform to UObject naming conventions.
+		if (!FName::IsValidXName(NewTextStr, INVALID_OBJECTNAME_CHARACTERS, &OutErrorMessage))
+		{
+			return false;
+		}
+
 		if (const UActorComponent* ComponentInstance = Data->GetComponentTemplate())
 		{
 			AActor* ExistingNameSearchScope = ComponentInstance->GetOwner();
