Defer re-running construction scripts until after we fix up the template hierarchy when duplicating a component in the Blueprint editor.

diff --git a/Engine/Source/Editor/SubobjectDataInterface/Private/SubobjectDataSubsystem.cpp b/Engine/Source/Editor/SubobjectDataInterface/Private/SubobjectDataSubsystem.cpp
--- a/Engine/Source/Editor/SubobjectDataInterface/Private/SubobjectDataSubsystem.cpp
+++ b/Engine/Source/Editor/SubobjectDataInterface/Private/SubobjectDataSubsystem.cpp
@@ -2417,6 +2417,13 @@ void USubobjectDataSubsystem::DuplicateSubobjects(const FSubobjectDataHandle& Co
 	NewSubobjectParams.BlueprintContext = BpContext;
 	NewSubobjectParams.ParentHandle = Context;
 	NewSubobjectParams.bConformTransformToParent = false;
+
+	// If we have a valid BP context, defer this step until after the AddNewSubobject() call, because we want
+	// to first fix up the template hierarchy (below) before we re-run construction scripts on any instances.
+	if (BpContext)
+	{
+		NewSubobjectParams.bSkipMarkBlueprintModified = true;
+	}
 	
 	FText FailedAddReason = FText::GetEmpty();
 	
@@ -2499,6 +2506,13 @@ void USubobjectDataSubsystem::DuplicateSubobjects(const FSubobjectDataHandle& Co
 			}
 		}
 	}
+
+	// Now that the hierarchy has been fixed up, we can go ahead and mark the BP as modified (if valid). This in turn will re-run
+	// construction scripts on any instances of the Blueprint, whose hierarchies will also now include the new (duplicate) subobject.
+	if (BpContext)
+	{
+		FBlueprintEditorUtils::MarkBlueprintAsStructurallyModified(BpContext);
+	}
 }
 
 FScopedTransaction* USubobjectDataSubsystem::BeginTransaction(const TArray<FSubobjectDataHandle>& Handles, const FText& Description, UBlueprint* InBlueprint)
