Fix a bug where you could not delete an instance-added subobject after it was added to an actor

#rnx

diff --git a/Engine/Source/Editor/SubobjectDataInterface/Private/SubobjectData.cpp b/Engine/Source/Editor/SubobjectDataInterface/Private/SubobjectData.cpp
--- a/Engine/Source/Editor/SubobjectDataInterface/Private/SubobjectData.cpp
+++ b/Engine/Source/Editor/SubobjectDataInterface/Private/SubobjectData.cpp
@@ -77,6 +77,11 @@ bool FSubobjectData::CanDelete() const
 			// Otherwise you can't delete the default scene root
 			return false;
 		}
+		else
+		{
+			// You can delete any instance-added component
+			return IsInstancedComponent();
+		}
 	}
 	
 	// Otherwise, it can't be deleted
