When recursively creating parent subobjects, pass in the correct component template instead of incorrectly passing the same one. This fixes a bug where a native subobject could be attached to itself, and display as a duplicate.

#rnx

diff --git a/Engine/Source/Editor/SubobjectDataInterface/Private/SubobjectDataSubsystem.cpp b/Engine/Source/Editor/SubobjectDataInterface/Private/SubobjectDataSubsystem.cpp
--- a/Engine/Source/Editor/SubobjectDataInterface/Private/SubobjectDataSubsystem.cpp
+++ b/Engine/Source/Editor/SubobjectDataInterface/Private/SubobjectDataSubsystem.cpp
@@ -134,7 +134,7 @@ FSubobjectDataHandle USubobjectDataSubsystem::FindOrCreateAttachParentForCompone
 			if (!ParentHandle.IsValid())
 			{
 				// Recursively add the parent handle to the tree if it does not exist yet
-				ParentHandle = FactoryCreateSubobjectDataWithParent(SceneComponent, FindOrCreateAttachParentForComponent(
+				ParentHandle = FactoryCreateSubobjectDataWithParent(SceneComponent->GetAttachParent(), FindOrCreateAttachParentForComponent(
 					SceneComponent->GetAttachParent(),
 					ActorRootHandle,
 					ExistingHandles)
