Back out CL 21894681 for 5.1.

diff --git a/Engine/Source/Editor/SubobjectDataInterface/Private/SubobjectDataSubsystem.cpp b/Engine/Source/Editor/SubobjectDataInterface/Private/SubobjectDataSubsystem.cpp
--- a/Engine/Source/Editor/SubobjectDataInterface/Private/SubobjectDataSubsystem.cpp
+++ b/Engine/Source/Editor/SubobjectDataInterface/Private/SubobjectDataSubsystem.cpp
@@ -2086,12 +2086,6 @@ bool USubobjectDataSubsystem::IsValidRename(const FSubobjectDataHandle& Handle,
 			return true;
 		}
 
-		// Subobject names must conform to UObject naming conventions.
-		if (!FName::IsValidXName(NewTextStr, INVALID_OBJECTNAME_CHARACTERS, &OutErrorMessage))
-		{
-			return false;
-		}
-
 		if (const UActorComponent* ComponentInstance = Data->GetComponentTemplate())
 		{
 			AActor* ExistingNameSearchScope = ComponentInstance->GetOwner();
