Prevent a crash when attempting to add an invalid subobject class type with the Subobject Data Subsystem. In an ideal world, a subobject does not have to be a component, but that is a current limitation of the UE object model that is exposed to blueprints.

diff --git a/Engine/Source/Editor/SubobjectDataInterface/Private/SubobjectDataSubsystem.cpp b/Engine/Source/Editor/SubobjectDataInterface/Private/SubobjectDataSubsystem.cpp
--- a/Engine/Source/Editor/SubobjectDataInterface/Private/SubobjectDataSubsystem.cpp
+++ b/Engine/Source/Editor/SubobjectDataInterface/Private/SubobjectDataSubsystem.cpp
@@ -843,6 +843,16 @@ FSubobjectDataHandle USubobjectDataSubsystem::AddNewSubobject(const FAddNewSubob
 		FailReason =  LOCTEXT("AddComponentFailed", "Cannot add components that have \"Within\" markup");
 		return NewDataHandle;
 	}
+
+	// Ensure that the new class is actually a valid subobject type.
+	// As of right now, that means the class has to be a child of an Actor Component. 
+	if (!NewClass->IsChildOf(UActorComponent::StaticClass()))
+	{
+		FFormatNamedArguments Arguments;
+		Arguments.Add(TEXT("ClassName"), NewClass->GetDisplayNameText());
+		FailReason = FText::Format(LOCTEXT("AddComponentFailed_InvalidClassType", "Cannot add a subobject of class '{ClassName}'"), Arguments);
+		return NewDataHandle;
+	}
 	
 	FName TemplateVariableName;
 	const USCS_Node* SCSNode = Cast<const USCS_Node>(Asset);
