Fix subobject editor allowinginstanced components to have a dot in its name

diff --git a/Engine/Source/Editor/SubobjectDataInterface/Private/SubobjectDataSubsystem.cpp b/Engine/Source/Editor/SubobjectDataInterface/Private/SubobjectDataSubsystem.cpp
--- a/Engine/Source/Editor/SubobjectDataInterface/Private/SubobjectDataSubsystem.cpp
+++ b/Engine/Source/Editor/SubobjectDataInterface/Private/SubobjectDataSubsystem.cpp
@@ -1286,13 +1286,13 @@ int32 USubobjectDataSubsystem::DeleteSubobject(const FSubobjectDataHandle& Conte
 bool USubobjectDataSubsystem::RenameSubobject(const FSubobjectDataHandle& Handle, const FText& InNewName)
 {
 	FText OutErrorMessage;
-	if(!IsValidRename(Handle, InNewName, OutErrorMessage))
+	if (!IsValidRename(Handle, InNewName, OutErrorMessage))
 	{
 		return false;	
 	}
 
 	const FSubobjectData* Data = Handle.GetData();
-	if(!Data)
+	if (!Data)
 	{
 		return false;
 	}
@@ -1309,44 +1309,40 @@ bool USubobjectDataSubsystem::RenameSubobject(const FSubobjectDataHandle& Handle
 	}
 	
 	// For instanced components
-	if(UActorComponent* ComponentInstance = Data->GetMutableComponentTemplate())
+	if (UActorComponent* ComponentInstance = Data->GetMutableComponentTemplate())
 	{
-		if(Data->IsInstancedComponent())
+		UBlueprint* const BP = Data->GetBlueprint();
+		if (!BP)
+		{
+			return false;
+		}
+		
+		const FString DesiredName = InNewName.ToString();
+		const FName ValidatedNewName = FKismetNameValidator(BP).IsValid(DesiredName) == EValidatorResult::Ok
+			? FName(DesiredName)
+			: FBlueprintEditorUtils::FindUniqueKismetName(BP, DesiredName);
+		
+		if (Data->IsInstancedComponent())
 		{
-			ERenameFlags RenameFlags = REN_DontCreateRedirectors;
-	
 			// name collision could occur due to e.g. our archetype being updated and causing a conflict with our ComponentInstance:
-			FString NewNameAsString = InNewName.ToString();
-			if(StaticFindObject(UObject::StaticClass(), ComponentInstance->GetOuter(), *NewNameAsString) == nullptr)
+			const FString NewNameAsString = ValidatedNewName.ToString();
+			if (StaticFindObject(UObject::StaticClass(), ComponentInstance->GetOuter(), *NewNameAsString) == nullptr)
 			{
+				constexpr ERenameFlags RenameFlags = REN_DontCreateRedirectors;
 				ComponentInstance->Rename(*NewNameAsString, nullptr, RenameFlags);
 			}
 			return true;
 		}
-		else if(UBlueprint* BP = Data->GetBlueprint())
+		
+		// Is this desired name the same as what is already there? If so then don't bother
+		USCS_Node* const SCSNode = Data->GetSCSNode();
+		if(SCSNode && SCSNode->GetVariableName().ToString().Equals(DesiredName))
 		{
-			FName ValidatedNewName;
-			FString DesiredName = InNewName.ToString();
-			
-			// Is this desired name the same as what is already there? If so then don't bother
-			USCS_Node* SCSNode = Data->GetSCSNode();
-			if(SCSNode && SCSNode->GetVariableName().ToString().Equals(DesiredName))
-			{
-				return true;
-			}
-			
-			if (FKismetNameValidator(BP).IsValid(DesiredName) == EValidatorResult::Ok)
-			{
-				ValidatedNewName = FName(DesiredName);
-			}
-			else
-			{
-				ValidatedNewName = FBlueprintEditorUtils::FindUniqueKismetName(BP, DesiredName);
-			}
-			
-			FBlueprintEditorUtils::RenameComponentMemberVariable(BP, SCSNode, ValidatedNewName);
 			return true;
 		}
+			
+		FBlueprintEditorUtils::RenameComponentMemberVariable(BP, SCSNode, ValidatedNewName);
+		return true;
 	}
 	
 	return false;
