Fix fallout of not being able to rename instanced component of native C++ classes

diff --git a/Engine/Source/Editor/SubobjectDataInterface/Private/SubobjectDataSubsystem.cpp b/Engine/Source/Editor/SubobjectDataInterface/Private/SubobjectDataSubsystem.cpp
--- a/Engine/Source/Editor/SubobjectDataInterface/Private/SubobjectDataSubsystem.cpp
+++ b/Engine/Source/Editor/SubobjectDataInterface/Private/SubobjectDataSubsystem.cpp
@@ -1308,41 +1308,51 @@ bool USubobjectDataSubsystem::RenameSubobject(const FSubobjectDataHandle& Handle
 		}
 	}
 	
-	// For instanced components
+	// For components: either instanced or simple construction script ("Add Component" in Blueprint editor)
 	if (UActorComponent* ComponentInstance = Data->GetMutableComponentTemplate())
 	{
-		UBlueprint* const BP = Data->GetBlueprint();
-		if (!BP)
-		{
-			return false;
-		}
-		
+		UBlueprint* BP = Data->GetBlueprint();
 		const FString DesiredName = InNewName.ToString();
-		const FName ValidatedNewName = FKismetNameValidator(BP).IsValid(DesiredName) == EValidatorResult::Ok
-			? FName(DesiredName)
-			: FBlueprintEditorUtils::FindUniqueKismetName(BP, DesiredName);
+		auto ValidateNameInBP = [BP](const FString& DesiredName) 
+		{
+			return FKismetNameValidator(BP).IsValid(DesiredName) == EValidatorResult::Ok
+				? FName(DesiredName)
+				: FBlueprintEditorUtils::FindUniqueKismetName(BP, DesiredName);
+		};
 		
+		// When placed in level, just the component needs to be renamed.
 		if (Data->IsInstancedComponent())
 		{
-			// name collision could occur due to e.g. our archetype being updated and causing a conflict with our ComponentInstance:
-			const FString NewNameAsString = ValidatedNewName.ToString();
+			// UBlueprint instance required may not always be available, e.g. if we added a non-Blueprint C++ class into the level like AStaticMeshActor.
+			// If they are available, use FKismetNameValidator and otherwise we'll just do our best here and replace invalid characters.
+			const FString NewNameAsString = BP
+				? ValidateNameInBP(DesiredName).ToString()
+				: FBlueprintEditorUtils::ReplaceInvalidBlueprintNameCharactersInline(DesiredName);
+			
+			// Name collision could occur due to e.g. our archetype being updated and causing a conflict with our ComponentInstance:
 			if (StaticFindObject(UObject::StaticClass(), ComponentInstance->GetOuter(), *NewNameAsString) == nullptr)
 			{
-				constexpr ERenameFlags RenameFlags = REN_DontCreateRedirectors;
+				const ERenameFlags RenameFlags = REN_DontCreateRedirectors;
 				ComponentInstance->Rename(*NewNameAsString, nullptr, RenameFlags);
 			}
+			
 			return true;
 		}
-		
-		// Is this desired name the same as what is already there? If so then don't bother
-		USCS_Node* const SCSNode = Data->GetSCSNode();
-		if(SCSNode && SCSNode->GetVariableName().ToString().Equals(DesiredName))
+
+		// In Blueprint editor, the component variable name must be updated.
+		if (BP)
 		{
+			// Is this desired name the same as what is already there? If so then don't bother.
+			USCS_Node* SCSNode = Data->GetSCSNode();
+			if(SCSNode && SCSNode->GetVariableName().ToString().Equals(DesiredName))
+			{
+				return true;
+			}
+			
+			const FName ValidatedNewName = ValidateNameInBP(DesiredName);
+			FBlueprintEditorUtils::RenameComponentMemberVariable(BP, SCSNode, ValidatedNewName);
 			return true;
 		}
-			
-		FBlueprintEditorUtils::RenameComponentMemberVariable(BP, SCSNode, ValidatedNewName);
-		return true;
 	}
 	
 	return false;
diff --git a/Engine/Source/Editor/UnrealEd/Private/Kismet2/BlueprintEditorUtils.cpp b/Engine/Source/Editor/UnrealEd/Private/Kismet2/BlueprintEditorUtils.cpp
--- a/Engine/Source/Editor/UnrealEd/Private/Kismet2/BlueprintEditorUtils.cpp
+++ b/Engine/Source/Editor/UnrealEd/Private/Kismet2/BlueprintEditorUtils.cpp
@@ -7979,6 +7979,21 @@ FBlueprintMacroCosmeticInfo FBlueprintEditorUtils::GetCosmeticInfoForMacro(UEdGr
 	return FBlueprintMacroCosmeticInfo();
 }
 
+void FBlueprintEditorUtils::ReplaceInvalidBlueprintNameCharacters(FString& InBaseName)
+{
+	for (TCHAR& TestChar : InBaseName)
+	{
+		for (TCHAR BadChar : UE_BLUEPRINT_INVALID_NAME_CHARACTERS)
+		{
+			if (TestChar == BadChar)
+			{
+				TestChar = TEXT('_');
+				break;
+			}
+		}
+	}
+}
+
 FName FBlueprintEditorUtils::FindUniqueKismetName(const UBlueprint* InBlueprint, const FString& InBaseName, UStruct* InScope/* = nullptr*/)
 {
 	int32 Count = 0;
@@ -7993,17 +8008,7 @@ FName FBlueprintEditorUtils::FindUniqueKismetName(const UBlueprint* InBlueprint,
 	// Clean up BaseName to not contain any invalid characters, which will mean we can never find a legal name no matter how many numbers we add
 	if (Result == EValidatorResult::ContainsInvalidCharacters)
 	{
-		for (TCHAR& TestChar : BaseName)
-		{
-			for (TCHAR BadChar : UE_BLUEPRINT_INVALID_NAME_CHARACTERS)
-			{
-				if (TestChar == BadChar)
-				{
-					TestChar = TEXT('_');
-					break;
-				}
-			}
-		}
+		ReplaceInvalidBlueprintNameCharacters(BaseName);
 		KismetName = BaseName;
 		Result = NameValidator->IsValid(KismetName);
 	}
diff --git a/Engine/Source/Editor/UnrealEd/Public/Kismet2/BlueprintEditorUtils.h b/Engine/Source/Editor/UnrealEd/Public/Kismet2/BlueprintEditorUtils.h
--- a/Engine/Source/Editor/UnrealEd/Public/Kismet2/BlueprintEditorUtils.h
+++ b/Engine/Source/Editor/UnrealEd/Public/Kismet2/BlueprintEditorUtils.h
@@ -1543,6 +1543,21 @@ public:
 	 * @return					A unique name that will not conflict in the Blueprint
 	 */
 	static FName FindUniqueKismetName(const UBlueprint* InBlueprint, const FString& InBaseName, UStruct* InScope = NULL);
+	
+	/**
+	 * Cleanses a name of invalid characters and replaces them with '_'.
+	 * See UE_BLUEPRINT_INVALID_NAME_CHARACTERS for invalid characters.
+	 *
+	 * @param InBaseName The base name to assign
+	 */
+	static void ReplaceInvalidBlueprintNameCharacters(FString& InBaseName);
+	
+	/** Util version of ReplaceInvalidBlueprintNameCharacters that performs the operation inline. */
+	static FString ReplaceInvalidBlueprintNameCharactersInline(FString InBaseName)
+	{
+		ReplaceInvalidBlueprintNameCharacters(InBaseName);
+		return InBaseName;
+	}
 
 	/** Finds a unique and valid name for a custom event */
 	static FName FindUniqueCustomEventName(const UBlueprint* Blueprint);
