[BugFix] Components copy/pasted from an actor instance's SCS can't be deleted from a Blueprint

diff --git a/Engine/Source/Editor/SubobjectDataInterface/Private/SubobjectDataSubsystem.cpp b/Engine/Source/Editor/SubobjectDataInterface/Private/SubobjectDataSubsystem.cpp
--- a/Engine/Source/Editor/SubobjectDataInterface/Private/SubobjectDataSubsystem.cpp
+++ b/Engine/Source/Editor/SubobjectDataInterface/Private/SubobjectDataSubsystem.cpp
@@ -44,6 +44,17 @@
 
 DEFINE_LOG_CATEGORY_STATIC(LogSubobjectSubsystem, Log, All);
 
+namespace UE::SubobjectDataSubsystem
+{
+	bool bForcePastedComponentsToSCS = true;
+	FAutoConsoleVariableRef CVarAudioShapesEnabled(
+		TEXT("bp.bForcePastedComponentsToSCS"),
+		bForcePastedComponentsToSCS,
+		TEXT("Setting this to True will change instanced components pasted into blueprints to be SCS components"),
+		ECVF_Default);
+}
+
+
 /** Notify the Level Editor that there have been subobject changes to an instance and it needs to be refreshed */
 static void BroadcastInstanceChanges()
 {
@@ -2305,6 +2316,15 @@ void USubobjectDataSubsystem::PasteSubobjects(const FSubobjectDataHandle& PasteT
 			UActorComponent* NewActorComponent = NewObjectPair.Value;
 			check(NewActorComponent);
 			
+			// make sure creation method is set to SimpleConstructionScript
+			if (UE::SubobjectDataSubsystem::bForcePastedComponentsToSCS)
+			{
+				if (NewActorComponent->CreationMethod == EComponentCreationMethod::Instance)
+				{
+					NewActorComponent->CreationMethod = EComponentCreationMethod::SimpleConstructionScript;
+				}
+			}
+			
 			// Create a new SCS node to contain the new component and add it to the tree
 			USCS_Node* NewSCSNode = Blueprint->SimpleConstructionScript->CreateNodeAndRenameComponent(NewActorComponent);
 			if(NewSCSNode)
