Silence warning when subobject deletion is forced

diff --git a/Engine/Source/Editor/SubobjectDataInterface/Private/SubobjectDataSubsystem.cpp b/Engine/Source/Editor/SubobjectDataInterface/Private/SubobjectDataSubsystem.cpp
--- a/Engine/Source/Editor/SubobjectDataInterface/Private/SubobjectDataSubsystem.cpp
+++ b/Engine/Source/Editor/SubobjectDataInterface/Private/SubobjectDataSubsystem.cpp
@@ -1221,7 +1221,7 @@ int32 USubobjectDataSubsystem::DeleteSubobjects(const FSubobjectDataHandle& Cont
 								}
 							}
 						}
-						else
+						else if(!bForce)
 						{
 							UE_LOG(LogSubobjectSubsystem, Warning, TEXT("Cannot remove subobject '%s' because it is the default scene root!"), *Data->GetDisplayString());
 						}
