FSubobjectData, used by the component editors (in world component editor, blueprint editor) now properly stores the component reference for overridden and instance components

diff --git a/Engine/Source/Editor/SubobjectDataInterface/Private/SubobjectData.cpp b/Engine/Source/Editor/SubobjectDataInterface/Private/SubobjectData.cpp
--- a/Engine/Source/Editor/SubobjectDataInterface/Private/SubobjectData.cpp
+++ b/Engine/Source/Editor/SubobjectDataInterface/Private/SubobjectData.cpp
@@ -1045,12 +1045,39 @@ USCS_Node* FSubobjectData::FindSCSNodeForInstance(const UActorComponent* Instanc
 	return nullptr;
 }
 
+bool GUseLegacySubobjectTemplateBehavior = false;
+static FAutoConsoleVariableRef CVarUseLegacySubobjectTemplateBehavior(
+	TEXT("bp.UseLegacySubobjectTemplateBehavior"),
+	GUseLegacySubobjectTemplateBehavior,
+	TEXT("Use the legacy behavior for setting up subobject template objects.")
+);
+
 bool FSubobjectData::AttemptToSetSCSNode()
 {
 	if (USCS_Node* PossibleSCS = Cast<USCS_Node>(WeakObjectPtr.Get()))
 	{
-		WeakObjectPtr = PossibleSCS->ComponentTemplate;
 		SCSNodePtr = PossibleSCS;
+		if (GUseLegacySubobjectTemplateBehavior)
+		{
+			WeakObjectPtr = PossibleSCS->ComponentTemplate;
+		}
+		else
+		{
+			const FSubobjectData* RootSubobjectData = GetRootSubobject().GetData();
+			const UObject* Object = RootSubobjectData->GetObject();
+			UBlueprintGeneratedClass* BPGC = Cast<UBlueprintGeneratedClass>(Object->GetClass());
+			if (BPGC)
+			{
+				WeakObjectPtr = PossibleSCS->GetActualComponentTemplate(BPGC);
+			}
+			else
+			{
+				// I don't believe this is possible, but we may have strange clients
+				// and the best thing we can do for them is provide the current behavior:
+				ensureMsgf(false, TEXT("Falling back to cached Component Template, consider running bp.UseLegacySubobjectTemplateBehavior true"));
+				WeakObjectPtr = PossibleSCS->ComponentTemplate;
+			}
+		}
 		return true;
 	}
 	// If this is an instanced component, then we can find it's SCS node
@@ -1062,7 +1089,10 @@ bool FSubobjectData::AttemptToSetSCSNode()
 			SCSNodePtr = FindSCSNodeForInstance(Template, Template->GetOwner()->GetClass());
 			if (SCSNodePtr.IsValid())
 			{
-				WeakObjectPtr = SCSNodePtr->ComponentTemplate;
+				if (GUseLegacySubobjectTemplateBehavior)
+				{
+					WeakObjectPtr = SCSNodePtr->ComponentTemplate;
+				}
 				return true;
 			}
 		}
