Fix several regressions created by 27427693. With these fixes we can also copy components (via the blueprint editor context menu) that are overrides of SCS components (e.g. live in the ICH). This is the 'root' bug that requires I stick with 27427693 despite the instability it has introduced.

Regressions fixed:
Crash on duplicate of instance added component
Crash on copy of ICH component (in blueprint editor)
Ensure on reparent of level editor component to SCS or ICH component

Existing bus fixed:
Failure to reparent component in the level editor if it is changed in a blueprint that inherited the component from another blueprint (which added it via the blueprint editor)

diff --git a/Engine/Source/Editor/SubobjectDataInterface/Private/SubobjectData.cpp b/Engine/Source/Editor/SubobjectDataInterface/Private/SubobjectData.cpp
--- a/Engine/Source/Editor/SubobjectDataInterface/Private/SubobjectData.cpp
+++ b/Engine/Source/Editor/SubobjectDataInterface/Private/SubobjectData.cpp
@@ -210,7 +210,8 @@ UActorComponent* FSubobjectData::FindMutableComponentInstanceInActor(const AActo
 				}
 			}
 		}
-		else if (ComponentTemplate)
+		
+		if (ComponentTemplate)
 		{
 			TInlineComponentArray<UActorComponent*> Components;
 			InActor->GetComponents(Components);
diff --git a/Engine/Source/Editor/SubobjectDataInterface/Private/SubobjectDataSubsystem.cpp b/Engine/Source/Editor/SubobjectDataInterface/Private/SubobjectDataSubsystem.cpp
--- a/Engine/Source/Editor/SubobjectDataInterface/Private/SubobjectDataSubsystem.cpp
+++ b/Engine/Source/Editor/SubobjectDataInterface/Private/SubobjectDataSubsystem.cpp
@@ -2018,17 +2018,10 @@ bool USubobjectDataSubsystem::AttachSubobject(const FSubobjectDataHandle& OwnerH
 	if(ChildToAddData->IsComponent())
 	{
 		const UActorComponent* ComponentTemplate = OwnerData->GetObject<UActorComponent>();
-		USceneComponent* ParentInstance = nullptr;
-		bool bIsInstancedComponent = OwnerData->IsInstancedComponent();
-
-		if (ComponentTemplate)
-		{
-			// Find the component instance on the current actor if we can
-			const FSubobjectDataHandle& ActorHandle = GetActorRootHandle(OwnerData->GetHandle());
-			const FSubobjectData* ActorData = ActorHandle.GetData();
-			ParentInstance = ActorData ? Cast<USceneComponent>(OwnerData->FindMutableComponentInstanceInActor(ActorData->GetObject<AActor>())) : nullptr;
-			bIsInstancedComponent |= ParentInstance != nullptr;
-		}
+		const FSubobjectDataHandle& ActorHandle = GetActorRootHandle(OwnerData->GetHandle());
+		const FSubobjectData* ActorData = ActorHandle.GetData();
+		USceneComponent* ParentInstance = ActorData ? Cast<USceneComponent>(OwnerData->FindMutableComponentInstanceInActor(ActorData->GetObject<AActor>())) : nullptr;
+		const bool bIsInstancedComponent = OwnerData->IsInstancedComponent() || ParentInstance != nullptr;
 
 		// Add a child node to the SCS tree node if not already present
 		if(USCS_Node* SCS_ChildNode = ChildToAddData->GetSCSNode())
@@ -2474,7 +2467,6 @@ void USubobjectDataSubsystem::DuplicateSubobjects(const FSubobjectDataHandle& Co
 		if(UActorComponent* ComponentTemplate = OriginalData->GetMutableComponentTemplate())
 		{
 			USCS_Node* SCSNode = OriginalData->GetSCSNode();
-			check(SCSNode == nullptr || SCSNode->ComponentTemplate == ComponentTemplate);
 
 			NewSubobjectParams.NewClass = ComponentTemplate->GetClass();
 			if (BpContext)
@@ -2510,7 +2502,6 @@ void USubobjectDataSubsystem::DuplicateSubobjects(const FSubobjectDataHandle& Co
 		FSubobjectData* OriginalData = DuplicatedPair.Key;
 		FSubobjectData* NewData = DuplicatedPair.Value;
 
-		USceneComponent* OriginalComponent = CastChecked<USceneComponent>(OriginalData->GetMutableComponentTemplate());
 		USceneComponent* NewSceneComponent = CastChecked<USceneComponent>(NewData->GetMutableComponentTemplate());
 		
 		if(BpContext)
diff --git a/Engine/Source/Runtime/CoreUObject/Private/Blueprint/BlueprintSupport.cpp b/Engine/Source/Runtime/CoreUObject/Private/Blueprint/BlueprintSupport.cpp
--- a/Engine/Source/Runtime/CoreUObject/Private/Blueprint/BlueprintSupport.cpp
+++ b/Engine/Source/Runtime/CoreUObject/Private/Blueprint/BlueprintSupport.cpp
@@ -2806,7 +2806,7 @@ FObjectInitializer* FDeferredObjInitializationHelper::DeferObjectInitializerIfNe
 		// FObjectInitializers will init default sub-object properties, copying  
 		// from the super's DSOs) - this means that we need to separately defer 
 		// init'ing these sub-objects when their archetype hasn't been loaded yet 
-		else if (TargetObj->HasAnyFlags(RF_InheritableComponentTemplate))
+		else if (TargetObj->HasAllFlags(RF_InheritableComponentTemplate|RF_WasLoaded))
 		{
 			const UClass* OwnerClass = Cast<UClass>(TargetObj->GetOuter());
 			DEFERRED_DEPENDENCY_CHECK(OwnerClass && OwnerClass->HasAnyClassFlags(CLASS_CompiledFromBlueprint));
