Flatten FSubobjectData hierarchy to prevent slicing, the occurence of slicing necessitated removal of subtypes and associated machinery. They were not supportable.

diff --git a/Engine/Source/Editor/SubobjectDataInterface/Private/BaseSubobjectDataFactory.cpp b/Engine/Source/Editor/SubobjectDataInterface/Private/BaseSubobjectDataFactory.cpp
deleted file mode 100644
diff --git a/Engine/Source/Editor/SubobjectDataInterface/Private/ChildActorSubobjectData.cpp b/Engine/Source/Editor/SubobjectDataInterface/Private/ChildActorSubobjectData.cpp
deleted file mode 100644
diff --git a/Engine/Source/Editor/SubobjectDataInterface/Private/ChildSubobjectDataFactory.cpp b/Engine/Source/Editor/SubobjectDataInterface/Private/ChildSubobjectDataFactory.cpp
deleted file mode 100644
diff --git a/Engine/Source/Editor/SubobjectDataInterface/Private/InheritedSubobjectData.cpp b/Engine/Source/Editor/SubobjectDataInterface/Private/InheritedSubobjectData.cpp
deleted file mode 100644
diff --git a/Engine/Source/Editor/SubobjectDataInterface/Private/InheritedSubobjectDataFactory.cpp b/Engine/Source/Editor/SubobjectDataInterface/Private/InheritedSubobjectDataFactory.cpp
deleted file mode 100644
diff --git a/Engine/Source/Editor/SubobjectDataInterface/Private/SubobjectData.cpp b/Engine/Source/Editor/SubobjectDataInterface/Private/SubobjectData.cpp
--- a/Engine/Source/Editor/SubobjectDataInterface/Private/SubobjectData.cpp
+++ b/Engine/Source/Editor/SubobjectDataInterface/Private/SubobjectData.cpp
@@ -22,22 +22,68 @@ FSubobjectData::FSubobjectData()
 	// be generated if we are given an object.
 }
 
-FSubobjectData::FSubobjectData(UObject* ContextObject, const FSubobjectDataHandle& InParentHandle)
+FSubobjectData::FSubobjectData(UObject* ContextObject, const FSubobjectDataHandle& InParentHandle, const bool bInIsInheritedSCS)
 	: WeakObjectPtr(ContextObject)
 	, ParentObjectHandle(InParentHandle)
 	, SCSNodePtr(nullptr)
 {
-	AttemptToSetSCSNode();
-}
+	// Check for a BP added CAC
+	if (USCS_Node* SCS = Cast<USCS_Node>(ContextObject))
+	{
+		if (SCS->ComponentTemplate->IsA<UChildActorComponent>())
+		{
+			bIsChildActor = true;
+		}
+	}
+	else if (ContextObject && ContextObject->IsA<UChildActorComponent>())
+	{
+		bIsChildActor = true;
+	}
 
-FSubobjectData::~FSubobjectData()
-{
-	ChildrenHandles.Empty();
-	WeakObjectPtr = nullptr;
+	if (UActorComponent* Component = Cast<UActorComponent>(ContextObject))
+	{
+		// Create an inherited subobject data
+		if (bInIsInheritedSCS || Component->CreationMethod != EComponentCreationMethod::Instance)
+		{
+			bIsInheritedSubobject = true;
+		}
+	}
+	else if (USCS_Node* SCS = Cast<USCS_Node>(ContextObject))
+	{
+		if (UActorComponent* Comp = SCS->ComponentTemplate)
+		{
+			bIsInheritedSubobject = (Comp->CreationMethod != EComponentCreationMethod::Instance);
+		}
+	}
+
+	bIsInheritedSCS = bInIsInheritedSCS && bIsInheritedSubobject;
+
+	AttemptToSetSCSNode();
 }
 
 bool FSubobjectData::CanEdit() const
 {
+	if (bIsInheritedSubobject)
+	{
+		if (IsComponent())
+		{
+			if (IsInstancedInheritedComponent())
+			{
+				const UActorComponent* Template = GetComponentTemplate();
+				return (Template ? Template->IsEditableWhenInherited() : false);
+			}
+			else if (!IsNativeComponent())
+			{
+				USCS_Node* SCS_Node = GetSCSNode();
+				return (SCS_Node != nullptr);
+			}
+			else if (const UActorComponent* ComponentTemplate = GetComponentTemplate())
+			{
+				return FComponentEditorUtils::GetPropertyForEditableNativeComponent(ComponentTemplate) != nullptr;
+			}
+		}
+	}
+
 	// Actors are editable by default
 	if(const AActor* ActorContext = GetObject<AActor>())
 	{
@@ -57,6 +103,16 @@ bool FSubobjectData::CanEdit() const
 
 bool FSubobjectData::CanDelete() const
 {
+	if (bIsInheritedSubobject)
+	{
+		if (IsInheritedComponent() || (IsDefaultSceneRoot() && SceneRootHasDefaultName()) || (GetSCSNode() != nullptr && IsInstancedInheritedComponent()) || IsChildActorSubtreeObject())
+		{
+			return false;
+		}
+
+		return true;
+	}
+
 	// Components can be deleted if they are not inherited
 	if (const UActorComponent* ComponentTemplate = GetComponentTemplate())
 	{
@@ -415,6 +471,13 @@ FText FSubobjectData::GetDisplayNameContextModifiers(bool bShowNativeComponentNa
 
 FText FSubobjectData::GetDisplayName() const
 {
+	if (bIsChildActor)
+	{
+		if (const UChildActorComponent* CAC = GetChildActorComponent())
+		{
+			return CAC->GetClass()->GetDisplayNameText();
+		}
+	}
 	if (const AActor* DefaultActor = GetObject<AActor>())
 	{
 		FString Name;
@@ -783,6 +846,13 @@ FText FSubobjectData::GetIntroducedInToolTipText() const
 
 FText FSubobjectData::GetActorDisplayText() const
 {
+	if (bIsChildActor)
+	{
+		if (const AActor* ChildActor = GetObject<AActor>())
+		{
+			return ChildActor->GetClass()->GetDisplayNameText();
+		}
+	}
 	if (const AActor* DefaultActor = GetObject<AActor>())
 	{
 		FString Name;
@@ -943,7 +1013,7 @@ bool FSubobjectData::IsComponent() const
 
 bool FSubobjectData::IsChildActor() const
 {
-	return false;
+	return bIsChildActor;
 }
 
 bool FSubobjectData::IsChildActorSubtreeObject() const
@@ -1078,7 +1148,12 @@ USCS_Node* FSubobjectData::GetSCSNode(bool bEvenIfPendingKill) const
 
 bool FSubobjectData::IsInheritedSCSNode() const
 {
-	return false;
+	return bIsInheritedSCS;
+}
+
+const UChildActorComponent* FSubobjectData::GetChildActorComponent(bool bEvenIfPendingKill) const
+{
+	return GetObject<UChildActorComponent>(bEvenIfPendingKill);
 }
 
 #undef LOCTEXT_NAMESPACE
diff --git a/Engine/Source/Editor/SubobjectDataInterface/Private/SubobjectDataBlueprintFunctionLibrary.cpp b/Engine/Source/Editor/SubobjectDataInterface/Private/SubobjectDataBlueprintFunctionLibrary.cpp
--- a/Engine/Source/Editor/SubobjectDataInterface/Private/SubobjectDataBlueprintFunctionLibrary.cpp
+++ b/Engine/Source/Editor/SubobjectDataInterface/Private/SubobjectDataBlueprintFunctionLibrary.cpp
@@ -6,9 +6,15 @@
 
 void USubobjectDataBlueprintFunctionLibrary::GetData(const FSubobjectDataHandle& DataHandle, FSubobjectData& OutData)
 {
+	// Copy the underlying subobject data - probably to the stack so that script can manipulate it
 	TSharedPtr<FSubobjectData> DataPtr = DataHandle.GetSharedDataPtr();
 	if(DataPtr.IsValid())
 	{
 		OutData = *DataPtr.Get();
 	}
 }
+
+const UObject* USubobjectDataBlueprintFunctionLibrary::GetObjectForBlueprint(const FSubobjectData& Data, UBlueprint* Blueprint) 
+{
+	return Data.GetObjectForBlueprint(Blueprint); 
+}
diff --git a/Engine/Source/Editor/SubobjectDataInterface/Private/SubobjectDataSubsystem.cpp b/Engine/Source/Editor/SubobjectDataInterface/Private/SubobjectDataSubsystem.cpp
--- a/Engine/Source/Editor/SubobjectDataInterface/Private/SubobjectDataSubsystem.cpp
+++ b/Engine/Source/Editor/SubobjectDataInterface/Private/SubobjectDataSubsystem.cpp
@@ -9,11 +9,6 @@
 
 #include "Editor/EditorEngine.h"		// FActorLabelUtilities
 #include "ComponentAssetBroker.h"		// FComponentAssetBrokerage
-#include "ChildActorSubobjectData.h"
-#include "InheritedSubobjectData.h"
-#include "BaseSubobjectDataFactory.h"
-#include "ChildSubobjectDataFactory.h"
-#include "InheritedSubobjectDataFactory.h"
 #include "Serialization/ArchiveReplaceOrClearExternalReferences.h"
 
 #include "BlueprintEditorSettings.h"	// bHideConstructionScriptComponentsInDetailsView
@@ -69,21 +64,10 @@ static void BroadcastInstanceChanges()
 
 void USubobjectDataSubsystem::Initialize(FSubsystemCollectionBase& Collection)
 {
-	FactoryManager = new FSubobjectFactoryManager();
-	check(FactoryManager);
-
-	FactoryManager->RegisterFactory(MakeShared<FBaseSubobjectDataFactory>());
-	FactoryManager->RegisterFactory(MakeShared<FChildSubobjectDataFactory>());
-	FactoryManager->RegisterFactory(MakeShared<FInheritedSubobjectDataFactory>());
 }
 
 void USubobjectDataSubsystem::Deinitialize()
 {
-	if(ensure(FactoryManager))
-	{
-		delete FactoryManager;
-		FactoryManager = nullptr;
-	}
 }
 
 USubobjectDataSubsystem* USubobjectDataSubsystem::Get()
@@ -2366,14 +2350,14 @@ void USubobjectDataSubsystem::PasteSubobjects(const FSubobjectDataHandle& PasteT
 				}
 			}
 
-			TSharedPtr<FInheritedSubobjectData> TargetData = StaticCastSharedPtr<FInheritedSubobjectData>(TargetParentHandle.GetSharedDataPtr());
+			TSharedPtr<FSubobjectData> TargetData = TargetParentHandle.GetSharedDataPtr();
 			
 			// Create a new subobject data set with this component. Use the SCS node here and the subobject data
 			// will correctly associate the component template
 			FSubobjectDataHandle NewDataHandle = FactoryCreateSubobjectDataWithParent(
 				NewSCSNode,
 				TargetParentHandle,
-				TargetData ? TargetData->bIsInheritedSCS : false
+				TargetData ? TargetData->IsInheritedSCSNode() : false
 			);
 
 			AttachSubobject(TargetParentHandle, NewDataHandle);
@@ -2580,19 +2564,12 @@ FScopedTransaction* USubobjectDataSubsystem::BeginTransaction(const TArray<FSubo
 
 FSubobjectDataHandle USubobjectDataSubsystem::CreateSubobjectData(UObject* Context, const FSubobjectDataHandle& ParentHandle /* = FSubobjectDataHandle::InvalidHandle */, bool bIsInheritedSCS/* = false */)
 {
-	FCreateSubobjectParams Params;
-	Params.Context = Context;
-	Params.ParentHandle = ParentHandle;
-	Params.bIsInheritedSCS = bIsInheritedSCS;
-
-	ISubobjectDataFactory* Factory = FactoryManager->FindFactoryToUse(Params);
-	check(Factory);
-	TSharedPtr<FSubobjectData> SharedPtr = Factory->CreateSubobjectData(Params);
+	TSharedPtr<FSubobjectData> SharedPtr = TSharedPtr<FSubobjectData>(new FSubobjectData(Context, ParentHandle, bIsInheritedSCS));
 	
 	if(!SharedPtr.IsValid())
 	{
 		ensureMsgf(false, TEXT("The subobject data factories failed to create subobject data!"));
-		SharedPtr = TSharedPtr<FSubobjectData>(new FSubobjectData(Context, ParentHandle));
+		SharedPtr = TSharedPtr<FSubobjectData>(new FSubobjectData(Context, ParentHandle, bIsInheritedSCS));
 	}
 
 	check(SharedPtr.IsValid());
@@ -2619,7 +2596,7 @@ FSubobjectDataHandle USubobjectDataSubsystem::FactoryCreateSubobjectDataWithPare
 	}
 
 	// Otherwise, we need to create a new handle
-	FSubobjectDataHandle OutHandle = CreateSubobjectData(Context, ParentHandle);
+	FSubobjectDataHandle OutHandle = CreateSubobjectData(Context, ParentHandle, bIsInheritedSCS);
 	
 	// Inform the parent that it has a new child
 	const bool bSuccess = ParentData->AddChildHandleOnly(OutHandle);
@@ -2637,16 +2614,9 @@ FSubobjectDataHandle USubobjectDataSubsystem::FactoryCreateInheritedBpSubobject(
 	
 	check(InSCSNode && ParentData && ParentData->IsValid());
 	// Get a handle for us to work with
-	FSubobjectDataHandle OutHandle = FactoryCreateSubobjectDataWithParent(InSCSNode, ParentHandle);
+	FSubobjectDataHandle OutHandle = FactoryCreateSubobjectDataWithParent(InSCSNode, ParentHandle, bIsInherited);
 	check(OutHandle.IsValid());
 	FSubobjectData* NewData = OutHandle.GetData();
-
-	// @todo the concept of inherited SCS nodes will not exist once the
-	// larger subobject refactor roles out, but this is necessary for now.
-	if(TSharedPtr<FInheritedSubobjectData> InheritedData = StaticCastSharedPtr<FInheritedSubobjectData>(OutHandle.GetSharedDataPtr()))
-	{
-		InheritedData->bIsInheritedSCS = bIsInherited;
-	}
 	
 	// Determine whether or not the given node is inherited from a parent Blueprint
 	USimpleConstructionScript* NodeSCS = InSCSNode->GetSCS();
diff --git a/Engine/Source/Editor/SubobjectDataInterface/Public/BaseSubobjectDataFactory.h b/Engine/Source/Editor/SubobjectDataInterface/Public/BaseSubobjectDataFactory.h
deleted file mode 100644
diff --git a/Engine/Source/Editor/SubobjectDataInterface/Public/ChildActorSubobjectData.h b/Engine/Source/Editor/SubobjectDataInterface/Public/ChildActorSubobjectData.h
deleted file mode 100644
diff --git a/Engine/Source/Editor/SubobjectDataInterface/Public/ChildSubobjectDataFactory.h b/Engine/Source/Editor/SubobjectDataInterface/Public/ChildSubobjectDataFactory.h
deleted file mode 100644
diff --git a/Engine/Source/Editor/SubobjectDataInterface/Public/InheritedSubobjectData.h b/Engine/Source/Editor/SubobjectDataInterface/Public/InheritedSubobjectData.h
deleted file mode 100644
diff --git a/Engine/Source/Editor/SubobjectDataInterface/Public/InheritedSubobjectDataFactory.h b/Engine/Source/Editor/SubobjectDataInterface/Public/InheritedSubobjectDataFactory.h
deleted file mode 100644
diff --git a/Engine/Source/Editor/SubobjectDataInterface/Public/SubobjectData.h b/Engine/Source/Editor/SubobjectDataInterface/Public/SubobjectData.h
--- a/Engine/Source/Editor/SubobjectDataInterface/Public/SubobjectData.h
+++ b/Engine/Source/Editor/SubobjectDataInterface/Public/SubobjectData.h
@@ -21,17 +21,14 @@ class UBlueprint;
 * If you wish to modify a subobject, use the SubobjectDataSubsystem.
 */
 USTRUCT(BlueprintType)
-struct SUBOBJECTDATAINTERFACE_API FSubobjectData
+struct SUBOBJECTDATAINTERFACE_API FSubobjectData final
 {
 friend class USubobjectDataSubsystem;
 public:
 	GENERATED_BODY()
 
 	FSubobjectData();
-
-	FSubobjectData(UObject* ContextObject, const FSubobjectDataHandle& ParentHandle);
-
-	virtual ~FSubobjectData();
+	FSubobjectData(UObject* ContextObject, const FSubobjectDataHandle& ParentHandle, const bool bIsInheritedSCS);
 
 	/** True if this subobject's handle is valid */
 	bool IsValid() const { return Handle.IsValid(); }
@@ -39,34 +36,34 @@ public:
 	/**
 	 * @return Whether or not we can edit properties for this subobject
 	 */
-	virtual bool CanEdit() const;
+	bool CanEdit() const;
 
 	/**
 	 * @return Whether or not this object represents a subobject that can be deleted
 	 */
-	virtual bool CanDelete() const;
+	bool CanDelete() const;
 
 	/**
 	* @return Whether or not this object represents a subobject that can be duplicated
 	*/
-	virtual bool CanDuplicate() const;
+	bool CanDuplicate() const;
 
 	/**
 	* @return Whether or not this object represents a subobject that can be copied
 	*/
-	virtual bool CanCopy() const;
+	bool CanCopy() const;
 
 	/**
 	 * @return Whether or not this object represents a subobject that can 
 	 * be reparented to other subobjects based on its context.
 	 */
-	virtual bool CanReparent() const;
+	bool CanReparent() const;
 
 	/**
 	* @return Whether or not this object represents a subobject that can 
 	* be renamed.
 	*/
-	virtual bool CanRename() const;
+	bool CanRename() const;
 
 	/**
 	* @return	A read-only reference to the object represented by this node.
@@ -119,30 +116,30 @@ public:
 
 	UBlueprint* GetBlueprint() const;
 
-	virtual bool IsInstancedComponent() const;
-	virtual bool IsInstancedActor() const;
-	virtual bool IsNativeComponent() const;
-	virtual bool IsBlueprintInheritedComponent() const;
-	virtual bool IsInheritedComponent() const;
-	virtual bool IsSceneComponent() const;
-	virtual bool IsRootComponent() const;
-	virtual bool IsDefaultSceneRoot() const;
+	bool IsInstancedComponent() const;
+	bool IsInstancedActor() const;
+	bool IsNativeComponent() const;
+	bool IsBlueprintInheritedComponent() const;
+	bool IsInheritedComponent() const;
+	bool IsSceneComponent() const;
+	bool IsRootComponent() const;
+	bool IsDefaultSceneRoot() const;
 	/** Returns true if the component template's FName starts with USceneComponent::GetDefaultSceneRootVariableName */
 	bool SceneRootHasDefaultName() const;
 	/* Returns true if this subobject is a component. */
-	virtual bool IsComponent() const;
-	virtual bool IsChildActor() const;
+	bool IsComponent() const;
+	bool IsChildActor() const;
 
 	/** Returns true if this subobject is a part of a child actor's hierarchy and not the root actor.  */
-	virtual bool IsChildActorSubtreeObject() const;
-	virtual bool IsRootActor() const;
-	virtual bool IsActor() const;
+	bool IsChildActorSubtreeObject() const;
+	bool IsRootActor() const;
+	bool IsActor() const;
 	/** True if this subobject is an instance of an inherited component (its owner is an instanced actor) */
-	virtual bool IsInstancedInheritedComponent() const;
+	bool IsInstancedInheritedComponent() const;
 
 	bool IsAttachedTo(const FSubobjectDataHandle& InHandle) const;
 
-	virtual FString GetDisplayString(bool bShowNativeComponentNames = true) const;
+	FString GetDisplayString(bool bShowNativeComponentNames = true) const;
 
 	/** Get the display name of this subobject that should be used during drag/drop operations */
 	FText GetDragDropDisplayText() const;
@@ -152,11 +149,11 @@ public:
 	 * components it will return "(Inherited)". For native components
 	 * this function will return "(NativeComponentName) (Inherited)"
 	 */
-	virtual FText GetDisplayNameContextModifiers(bool bShowNativeComponentNames = true) const;
+	FText GetDisplayNameContextModifiers(bool bShowNativeComponentNames = true) const;
 	
-	virtual FText GetDisplayName() const;
+	FText GetDisplayName() const;
 
-	virtual FName GetVariableName() const;
+	FName GetVariableName() const;
 
 	// Sockets for attaching in the viewport
 	FText GetSocketName() const;
@@ -211,10 +208,9 @@ public:
 	*/
 	FText GetIntroducedInToolTipText() const;
 
-	virtual FText GetActorDisplayText() const;
+	FText GetActorDisplayText() const;
 		
-protected:
-	
+private:
 	/**
 	* Add the given subobject handle as a child of this.
 	* NOTE: This does NOT do any actual manipulation of the subobject structure.
@@ -279,12 +275,11 @@ protected:
 
 	////////////////////////////////////////////////////
 	// Anything related to an SCS node will be changed with an upcoming refactor and should remain private
-protected:
-
+private:
 	// Tries to find a SCS node that was likely responsible for creating the specified instance component.  Note: This is not always possible to do!
 	static USCS_Node* FindSCSNodeForInstance(const UActorComponent* InstanceComponent, UClass* ClassToSearch);
 
-	virtual USCS_Node* GetSCSNode(bool bEvenIfPendingKill = false) const;
+	USCS_Node* GetSCSNode(bool bEvenIfPendingKill = false) const;
 
 	/**
 	* Attempts to set the SCS node member variable based on the WeakObjectPtr.
@@ -306,5 +301,11 @@ protected:
 	 * True if this SCS node's blueprint is a child of another Blueprint-generated class,
 	 * which means that we must use the InheritableComponentHandler
 	 */
-	virtual bool IsInheritedSCSNode() const;
+	bool IsInheritedSCSNode() const;
+	const UChildActorComponent* GetChildActorComponent(bool bEvenIfPendingKill = false) const;
+
+	/** True if this SCS node is inherited from another blueprint generated class. */
+	bool bIsInheritedSubobject = false;
+	bool bIsInheritedSCS = false;
+	bool bIsChildActor = false;
 };
\ No newline at end of file
diff --git a/Engine/Source/Editor/SubobjectDataInterface/Public/SubobjectDataBlueprintFunctionLibrary.h b/Engine/Source/Editor/SubobjectDataInterface/Public/SubobjectDataBlueprintFunctionLibrary.h
--- a/Engine/Source/Editor/SubobjectDataInterface/Public/SubobjectDataBlueprintFunctionLibrary.h
+++ b/Engine/Source/Editor/SubobjectDataInterface/Public/SubobjectDataBlueprintFunctionLibrary.h
@@ -88,7 +88,7 @@ public:
 	static const UObject* GetObject(const FSubobjectData& Data, bool bEvenIfPendingKill = false) { return Data.GetObject(bEvenIfPendingKill); }
 
 	UFUNCTION(BlueprintCallable, Category="Subobject Data")
-	static const UObject* GetObjectForBlueprint(const FSubobjectData& Data, UBlueprint* Blueprint) { return Data.GetObjectForBlueprint(Blueprint); }
+	static const UObject* GetObjectForBlueprint(const FSubobjectData& Data, UBlueprint* Blueprint);
 
 	UFUNCTION(BlueprintCallable, BlueprintPure, Category="Subobject Data")
 	UBlueprint* GetBlueprint(const FSubobjectData& Data) { return Data.GetBlueprint(); }
diff --git a/Engine/Source/Editor/SubobjectDataInterface/Public/SubobjectDataFactory.h b/Engine/Source/Editor/SubobjectDataInterface/Public/SubobjectDataFactory.h
deleted file mode 100644
diff --git a/Engine/Source/Editor/SubobjectDataInterface/Public/SubobjectDataSubsystem.h b/Engine/Source/Editor/SubobjectDataInterface/Public/SubobjectDataSubsystem.h
--- a/Engine/Source/Editor/SubobjectDataInterface/Public/SubobjectDataSubsystem.h
+++ b/Engine/Source/Editor/SubobjectDataInterface/Public/SubobjectDataSubsystem.h
@@ -7,7 +7,6 @@
 #include "Subsystems/EngineSubsystem.h"
 #include "SubobjectData.h"
 #include "ModuleDescriptor.h"
-#include "SubobjectDataFactory.h"
 #include "SubobjectDataSubsystem.generated.h"
 
 class UActorComponent;
@@ -419,12 +418,7 @@ private:
 	
 	/** Find the parent that this new subobject should attach to */
 	FSubobjectDataHandle FindParentForNewSubobject(const UObject* NewSubobject, const FSubobjectDataHandle& SelectedParent);
-	
-	FSubobjectFactoryManager* FactoryManager = nullptr;
 
 	FOnNewSubobjectAdded OnNewSubobjectAdded_Delegate;
 	
-public:
-	FSubobjectFactoryManager* GetSubobjectFactoryManager() const { return FactoryManager; }
-	
 };
\ No newline at end of file
