Remove overzealous ensure

diff --git a/Engine/Source/Editor/SubobjectDataInterface/Private/SubobjectDataSubsystem.cpp b/Engine/Source/Editor/SubobjectDataInterface/Private/SubobjectDataSubsystem.cpp
--- a/Engine/Source/Editor/SubobjectDataInterface/Private/SubobjectDataSubsystem.cpp
+++ b/Engine/Source/Editor/SubobjectDataInterface/Private/SubobjectDataSubsystem.cpp
@@ -289,8 +289,7 @@ void USubobjectDataSubsystem::GatherSubobjectData(UObject* Context, TArray<FSubo
 						if (USceneComponent* ParentComponent = SCS_Node->GetParentComponentTemplate(ActorBP))
 						{
 							FSubobjectDataHandle ParentHandle = FindAttachParentForInheritedComponent(ParentComponent, RootActorHandle, ActorBP);
-
-							if (ensure(ParentHandle.IsValid()))
+							if (ParentHandle.IsValid())
 							{
 								NewHandle = FactoryCreateInheritedBpSubobject(SCS_Node, ParentHandle, /* bIsInherited = */ StackIndex > 0, OutArray);
 							}
