Fix for a crash when changing the class or a subobject marked as "Within" to a BP subclass instead.

diff --git a/Engine/Source/Editor/SubobjectDataInterface/Private/SubobjectDataSubsystem.cpp b/Engine/Source/Editor/SubobjectDataInterface/Private/SubobjectDataSubsystem.cpp
--- a/Engine/Source/Editor/SubobjectDataInterface/Private/SubobjectDataSubsystem.cpp
+++ b/Engine/Source/Editor/SubobjectDataInterface/Private/SubobjectDataSubsystem.cpp
@@ -1579,11 +1579,12 @@ bool USubobjectDataSubsystem::ChangeSubobjectClass(const FSubobjectDataHandle& H
 
 						// Avoid the new instance going in to the transaction buffer as we are managing that manually through the custom change record
 						TGuardValue<ITransaction*> SuppressTransaction(GUndo, nullptr);
-						UObject* NewInstance = NewObject<UObject>(GetTransientPackage(), NewClass, *ReplacementName, ArchetypeInstance->GetFlags());
+						UObject* NewInstance = NewObject<UObject>(GetTransientOuterForRename(const_cast<UClass*>(NewClass)), NewClass, *ReplacementName, ArchetypeInstance->GetFlags());
+
 						UEngine::CopyPropertiesForUnrelatedObjects(ArchetypeInstance, NewInstance);
 						ReplacementObjects.Add(NewInstance);
 
-						StaticDuplicateObject(ArchetypeInstance, GetTransientPackage(), *ReplacedName);
+						StaticDuplicateObject(ArchetypeInstance, GetTransientOuterForRename(const_cast<UClass*>(NewClass)), *ReplacedName);
 					}
 				}
 
