Update class's overridden components immediately upon creating the overrridden component, so that the class can be safely loaded if it is saved

diff --git a/Engine/Source/Editor/SubobjectDataInterface/Private/SubobjectDataSubsystem.cpp b/Engine/Source/Editor/SubobjectDataInterface/Private/SubobjectDataSubsystem.cpp
--- a/Engine/Source/Editor/SubobjectDataInterface/Private/SubobjectDataSubsystem.cpp
+++ b/Engine/Source/Editor/SubobjectDataInterface/Private/SubobjectDataSubsystem.cpp
@@ -1413,6 +1413,11 @@ bool USubobjectDataSubsystem::ChangeSubobjectClass(const FSubobjectDataHandle& H
 					if (bRemoveEntry)
 					{
 						BlueprintObj->ComponentClassOverrides.RemoveAllSwap([ComponentTemplateName](const FBPComponentClassOverride& CCOverride) { return (CCOverride.ComponentName == ComponentTemplateName); });
+						if (UBlueprintGeneratedClass* Class = Cast<UBlueprintGeneratedClass>(BlueprintObj->GeneratedClass))
+						{
+							Class->Modify();
+							Class->ComponentClassOverrides.RemoveAllSwap([ComponentTemplateName](const FBPComponentClassOverride& CCOverride) { return (CCOverride.ComponentName == ComponentTemplateName); });
+						}
 					}
 					else
 					{
@@ -1422,6 +1427,11 @@ bool USubobjectDataSubsystem::ChangeSubobjectClass(const FSubobjectDataHandle& H
 				else
 				{
 					BlueprintObj->ComponentClassOverrides.Emplace(FBPComponentClassOverride(ComponentTemplateName, NewClass));
+					if (UBlueprintGeneratedClass* Class = Cast<UBlueprintGeneratedClass>(BlueprintObj->GeneratedClass))
+					{
+						Class->Modify();
+						Class->ComponentClassOverrides.Emplace(FBPComponentClassOverride(ComponentTemplateName, NewClass));
+					}
 				}
 
 				// Custom transaction change that operates on the UBlueprint and replaces all instances of the subobject with one of the new class
