AddNewSubobject intelligently finds the blueprint being edited, if any

diff --git a/Engine/Source/Editor/SubobjectDataInterface/Private/SubobjectData.cpp b/Engine/Source/Editor/SubobjectDataInterface/Private/SubobjectData.cpp
--- a/Engine/Source/Editor/SubobjectDataInterface/Private/SubobjectData.cpp
+++ b/Engine/Source/Editor/SubobjectDataInterface/Private/SubobjectData.cpp
@@ -324,6 +324,26 @@ UBlueprint* FSubobjectData::GetBlueprint() const
 	return nullptr;
 }
 
+UBlueprint* FSubobjectData::GetBlueprintBeingEdited() const
+{
+	UObject* Object = WeakObjectPtr.Get();
+	if (!Object)
+	{
+		return nullptr;
+	}
+
+	if (UBlueprint* BP = Cast<UBlueprint>(Object))
+	{
+		return BP;
+	}
+	else if (Object->IsTemplate())
+	{
+		return GetBlueprint();
+	}
+
+	return nullptr;
+}
+
 FString FSubobjectData::GetDisplayString(bool bShowNativeComponentNames /* = true */) const
 {
 	FName VariableName = GetVariableName();
diff --git a/Engine/Source/Editor/SubobjectDataInterface/Private/SubobjectDataSubsystem.cpp b/Engine/Source/Editor/SubobjectDataInterface/Private/SubobjectDataSubsystem.cpp
--- a/Engine/Source/Editor/SubobjectDataInterface/Private/SubobjectDataSubsystem.cpp
+++ b/Engine/Source/Editor/SubobjectDataInterface/Private/SubobjectDataSubsystem.cpp
@@ -861,11 +861,17 @@ FSubobjectDataHandle USubobjectDataSubsystem::AddNewSubobject(const FAddNewSubob
 	{
 		Asset = nullptr;
 	}
+
+	UBlueprint* Blueprint = Params.BlueprintContext;
+	if (!Blueprint)
+	{
+		// maybe the parent handle has a bp for context:
+		Blueprint = ParentObjHandle.GetData()->GetBlueprintBeingEdited();
+	}
 	
-	if(Params.BlueprintContext)
+	if(Blueprint)
 	{
-		UBlueprint* Blueprint = Params.BlueprintContext;
-		check(Blueprint != nullptr && Blueprint->SimpleConstructionScript != nullptr);
+		check(Blueprint->SimpleConstructionScript != nullptr);
 		Blueprint->Modify();
 		SaveSCSCurrentState(Blueprint->SimpleConstructionScript);
 		UActorComponent* NewComponent = nullptr;
diff --git a/Engine/Source/Editor/SubobjectDataInterface/Public/SubobjectData.h b/Engine/Source/Editor/SubobjectDataInterface/Public/SubobjectData.h
--- a/Engine/Source/Editor/SubobjectDataInterface/Public/SubobjectData.h
+++ b/Engine/Source/Editor/SubobjectDataInterface/Public/SubobjectData.h
@@ -115,6 +115,7 @@ public:
 	inline const UActorComponent* FindComponentInstanceInActor(const AActor* InActor) const { return FindMutableComponentInstanceInActor(InActor); }
 
 	UBlueprint* GetBlueprint() const;
+	UBlueprint* GetBlueprintBeingEdited() const;
 
 	bool IsInstancedComponent() const;
 	bool IsInstancedActor() const;
