Expose "PasteSubobjects" to blueprints.

diff --git a/Engine/Source/Editor/SubobjectDataInterface/Public/SubobjectDataSubsystem.h b/Engine/Source/Editor/SubobjectDataInterface/Public/SubobjectDataSubsystem.h
--- a/Engine/Source/Editor/SubobjectDataInterface/Public/SubobjectDataSubsystem.h
+++ b/Engine/Source/Editor/SubobjectDataInterface/Public/SubobjectDataSubsystem.h
@@ -346,6 +346,15 @@ public:
 	UFUNCTION(BlueprintCallable, Category = "SubobjectDataSubsystem")
 	bool CanPasteSubobjects(const FSubobjectDataHandle& RootHandle, UBlueprint* BPContext = nullptr) const;
 
+	/**
+	* Pastes the given subobjects to the PasteToContext.
+	* 
+	* @param PasteToContext		The subobject to paste things onto
+	* @param NewParentHandles	Array of Subobject Handles that you would like to paste
+	* @param BpContext			Blueprint to use, if you are pasting to a blueprint. Null if pasting to an instanced object
+	* @param OutPastedHandles	Array populated with the handles to the newly pasted subobjects
+	*/
+	UFUNCTION(BlueprintCallable, Category = "SubobjectDataSubsystem")
 	void PasteSubobjects(const FSubobjectDataHandle& PasteToContext, const TArray<FSubobjectDataHandle>& NewParentHandles, UBlueprint* BpContext, TArray<FSubobjectDataHandle>& OutPastedHandles);
 	
 	/**
@@ -356,6 +365,7 @@ public:
 	 * @param BpContext			Pointer to the current blueprint context if necessary. Use nullptr if dealing with instances
 	 * @param OutNewSubobjects	Array that will be populated with any newly created subobjects
 	 */
+	UFUNCTION(BlueprintCallable, Category = "SubobjectDataSubsystem")
 	void DuplicateSubobjects(const FSubobjectDataHandle& Context, const TArray<FSubobjectDataHandle>& SubobjectsToDup, UBlueprint* BpContext, TArray<FSubobjectDataHandle>& OutNewSubobjects);
 	
 	FScopedTransaction* BeginTransaction(const TArray<FSubobjectDataHandle>& Handles, const FText& Description, UBlueprint* Blueprint);
