PR #11871: SubobjectDataInterface: Make BPFL functions static

#pf 664777f25565dc78dcf97f9b

diff --git a/Engine/Source/Editor/SubobjectDataInterface/Public/SubobjectDataBlueprintFunctionLibrary.h b/Engine/Source/Editor/SubobjectDataInterface/Public/SubobjectDataBlueprintFunctionLibrary.h
--- a/Engine/Source/Editor/SubobjectDataInterface/Public/SubobjectDataBlueprintFunctionLibrary.h
+++ b/Engine/Source/Editor/SubobjectDataInterface/Public/SubobjectDataBlueprintFunctionLibrary.h
@@ -42,7 +42,7 @@ public:
 	static void GetHandle(const FSubobjectData& Data, FSubobjectDataHandle& OutHandle) { OutHandle = Data.GetHandle(); }
 
 	UFUNCTION(BlueprintCallable, BlueprintPure, Category = "Subobject Data")
-	FText GetDisplayName(const FSubobjectData& Data) { return Data.GetDisplayName(); }
+	static FText GetDisplayName(const FSubobjectData& Data) { return Data.GetDisplayName(); }
 
 	UFUNCTION(BlueprintCallable, BlueprintPure, Category = "Subobject Data")
 	static FName GetVariableName(const FSubobjectData& Data) { return Data.GetVariableName(); }
@@ -91,7 +91,7 @@ public:
 	static const UObject* GetObjectForBlueprint(const FSubobjectData& Data, UBlueprint* Blueprint);
 
 	UFUNCTION(BlueprintCallable, BlueprintPure, Category="Subobject Data")
-	UBlueprint* GetBlueprint(const FSubobjectData& Data) { return Data.GetBlueprint(); }
+	static UBlueprint* GetBlueprint(const FSubobjectData& Data) { return Data.GetBlueprint(); }
 
 	UFUNCTION(BlueprintCallable, BlueprintPure, Category="Subobject Data")
 	static bool IsInstancedComponent(const FSubobjectData& Data) { return Data.IsInstancedComponent(); }
@@ -126,4 +126,4 @@ public:
 	
 	UFUNCTION(BlueprintCallable, BlueprintPure, Category="Subobject Data")
 	static bool IsActor(const FSubobjectData& Data) { return Data.IsActor(); }
-};
\ No newline at end of file
+};
