Fix a regression that incorrectly caused a source Blueprint to be modified when adding a new component to an instance of it in the editor.

#rnx

diff --git a/Engine/Source/Editor/SubobjectDataInterface/Private/SubobjectDataSubsystem.cpp b/Engine/Source/Editor/SubobjectDataInterface/Private/SubobjectDataSubsystem.cpp
--- a/Engine/Source/Editor/SubobjectDataInterface/Private/SubobjectDataSubsystem.cpp
+++ b/Engine/Source/Editor/SubobjectDataInterface/Private/SubobjectDataSubsystem.cpp
@@ -862,8 +862,15 @@ FSubobjectDataHandle USubobjectDataSubsystem::AddNewSubobject(const FAddNewSubob
 		Asset = nullptr;
 	}
 
+	// We check for an invalid parent handle above, so we assume it is valid here.
+	check(Params.ParentHandle.IsValid());
+
+	// Editor utility scripts may set the parent handle but not a Blueprint context. It is arguable whether this constitutes a valid use case; although
+	// for backwards compatibility we'll continue to support it. However, in both the BP and instance in-editor use cases, data nodes that correlate to
+	// an SCS node will use the SCS template data, which means that we'll find a valid Blueprint context in that case. For editor scripting, we want it
+	// to always modify the Blueprint, but for "normal" editor operations (e.g. drag-and-drop), we need to bypass this if the root is an Actor instance.
 	UBlueprint* Blueprint = Params.BlueprintContext;
-	if (!Blueprint)
+	if (!Blueprint && !GetActorRootHandle(ParentObjHandle).IsValid())
 	{
 		// maybe the parent handle has a bp for context:
 		Blueprint = ParentObjHandle.GetData()->GetBlueprintBeingEdited();
