Add function to get parent handle for a subobject data
#tests NA

diff --git a/Engine/Source/Editor/SubobjectDataInterface/Public/SubobjectDataBlueprintFunctionLibrary.h b/Engine/Source/Editor/SubobjectDataInterface/Public/SubobjectDataBlueprintFunctionLibrary.h
--- a/Engine/Source/Editor/SubobjectDataInterface/Public/SubobjectDataBlueprintFunctionLibrary.h
+++ b/Engine/Source/Editor/SubobjectDataInterface/Public/SubobjectDataBlueprintFunctionLibrary.h
@@ -41,6 +41,12 @@ public:
 	UFUNCTION(BlueprintCallable, BlueprintPure, Category = "Subobject Data")
 	static void GetHandle(const FSubobjectData& Data, FSubobjectDataHandle& OutHandle) { OutHandle = Data.GetHandle(); }
 
+	/**
+	* @return Get the parent handle for this subobject data
+	*/
+	UFUNCTION(BlueprintCallable, BlueprintPure, Category = "Subobject Data")
+	static void GetParentHandle(const FSubobjectData& Data, FSubobjectDataHandle& OutHandle) { OutHandle = Data.GetParentHandle(); }
+
 	UFUNCTION(BlueprintCallable, BlueprintPure, Category = "Subobject Data")
 	static FText GetDisplayName(const FSubobjectData& Data) { return Data.GetDisplayName(); }
 
