Keep track of visited subobject data handles in USubobjectDataSubsystem::GetActorRootHandle. This is a safegaurd to stop an editor freeze during re-parenting of components which cannot be repro'd internally. Add an error log to get some more info from users who may encounter this.

diff --git a/Engine/Source/Editor/SubobjectDataInterface/Private/SubobjectDataSubsystem.cpp b/Engine/Source/Editor/SubobjectDataInterface/Private/SubobjectDataSubsystem.cpp
--- a/Engine/Source/Editor/SubobjectDataInterface/Private/SubobjectDataSubsystem.cpp
+++ b/Engine/Source/Editor/SubobjectDataInterface/Private/SubobjectDataSubsystem.cpp
@@ -810,10 +810,30 @@ FSubobjectDataHandle USubobjectDataSubsystem::GetActorRootHandle(const FSubobjec
 	FSubobjectDataHandle CurrentHandle = StartingHandle;
 	FSubobjectData* CurrentData = CurrentHandle.GetData();
 
+	TSet<FSubobjectDataHandle> VisitedHandles;
+	VisitedHandles.Add(StartingHandle);
+	
 	while (CurrentData && !CurrentData->IsActor())
 	{
 		CurrentHandle = CurrentData->GetParentHandle();
 		CurrentData = CurrentHandle.GetData();
+
+		// If we have already visited this node, then exit out of this loop
+		if (VisitedHandles.Contains(CurrentHandle))
+		{
+			const UObject* CurrentObj = CurrentData ? CurrentData->GetObject() : nullptr;						
+
+			ensureMsgf(
+				false, 
+				TEXT("[%hs] Duplicate visits when getting that actor's root handle! Exiting to prevent an infinite loop. Duplicate object: '%s' Flags: %s"), 
+				__func__,
+				*GetNameSafe(CurrentObj),
+				*LexToString(CurrentObj ? CurrentObj->GetFlags() : RF_NoFlags));
+
+			break;
+		}
+
+		VisitedHandles.Add(CurrentHandle);
 	}
 
 	return CurrentHandle;
