Add an additional nullcheck to SetupAttachment to prevent crashes when clearing a socket on a non-scene component

diff --git a/Engine/Source/Editor/SubobjectDataInterface/Private/SubobjectData.cpp b/Engine/Source/Editor/SubobjectDataInterface/Private/SubobjectData.cpp
--- a/Engine/Source/Editor/SubobjectDataInterface/Private/SubobjectData.cpp
+++ b/Engine/Source/Editor/SubobjectDataInterface/Private/SubobjectData.cpp
@@ -578,7 +578,11 @@ void FSubobjectData::SetSocketName(FName InNewName)
 void FSubobjectData::SetupAttachment(FName SocketName, const FSubobjectDataHandle& AttachParentHandle)
 {
 	USceneComponent* SceneComponentTemplate = Cast<USceneComponent>(GetMutableComponentTemplate());
-
+	if (!SceneComponentTemplate)
+	{
+		return;
+	}
+	
 	FSubobjectData* AttachParentData = AttachParentHandle.GetData();
 	USceneComponent* AttachParent = AttachParentData ?
 		Cast<USceneComponent>(AttachParentData->GetMutableComponentTemplate()) :
