Fixed the blueprint component tree loosing its selection when editing a child's properties in the cooked editor.

This was caused by a failure to find the SCS node of the cooked template.  Without it we were using the component instances when restoring the selection after rebuilding the component tree, which would have different pointer values after being reconstructed by the construction script, causing selection restoration to fail and select the root node by default.

diff --git a/Engine/Source/Editor/SubobjectDataInterface/Private/SubobjectData.cpp b/Engine/Source/Editor/SubobjectDataInterface/Private/SubobjectData.cpp
--- a/Engine/Source/Editor/SubobjectDataInterface/Private/SubobjectData.cpp
+++ b/Engine/Source/Editor/SubobjectDataInterface/Private/SubobjectData.cpp
@@ -2,6 +2,7 @@
 
 #include "SubobjectData.h"
 #include "Engine/Blueprint.h"			// Casting to UBlueprint
+#include "Engine/BlueprintGeneratedClass.h"
 #include "Components/ChildActorComponent.h"
 #include "GameFramework/Actor.h"
 #include "Kismet2/ComponentEditorUtils.h"
@@ -1121,19 +1122,22 @@ USCS_Node* FSubobjectData::FindSCSNodeForInstance(const UActorComponent* Instanc
 {
 	if ((ClassToSearch != nullptr) && InstanceComponent->IsCreatedByConstructionScript())
 	{
-		for (UClass* TestClass = ClassToSearch; TestClass->ClassGeneratedBy != nullptr; TestClass = TestClass->GetSuperClass())
-		{
-			if (UBlueprint* TestBP = Cast<UBlueprint>(TestClass->ClassGeneratedBy))
+		USCS_Node* Result = nullptr;
+		UBlueprintGeneratedClass::ForEachGeneratedClassInHierarchy(ClassToSearch, [InstanceComponent, &Result](const UBlueprintGeneratedClass* BPGClass)
 			{
-				if (TestBP->SimpleConstructionScript != nullptr)
+				if (BPGClass->SimpleConstructionScript != nullptr)
 				{
-					if (USCS_Node* Result = TestBP->SimpleConstructionScript->FindSCSNode(InstanceComponent->GetFName()))
+					Result = BPGClass->SimpleConstructionScript->FindSCSNode(InstanceComponent->GetFName());
+					if (Result != nullptr)
 					{
-						return Result;
+						// Stop traversing when we find the node.
+						return false;
 					}
 				}
-			}
-		}
+
+				return true;
+			});
+		return Result;
 	}
 
 	return nullptr;
