UE-214106 - Editor Crash when adding components to a Editor Only BP Actor

Iteration over child nodes didn't have sufficient safeguards for any entry removal. After calling FactoryCreateInheritedBpSubobject, we need to ensure that the iteration index is still within bounds.

diff --git a/Engine/Source/Editor/SubobjectDataInterface/Private/SubobjectDataSubsystem.cpp b/Engine/Source/Editor/SubobjectDataInterface/Private/SubobjectDataSubsystem.cpp
--- a/Engine/Source/Editor/SubobjectDataInterface/Private/SubobjectDataSubsystem.cpp
+++ b/Engine/Source/Editor/SubobjectDataInterface/Private/SubobjectDataSubsystem.cpp
@@ -2718,6 +2718,12 @@ FSubobjectDataHandle USubobjectDataSubsystem::FactoryCreateInheritedBpSubobject(
 			FSubobjectDataHandle NewChildHandle = FactoryCreateInheritedBpSubobject(ChildNode, OutHandle, bIsInherited, OutArray);
 			ensure(NewChildHandle.IsValid());
 			OutArray.Add(NewChildHandle);
+
+			// Since the child array can be modified, we need to ensure that our index remains in bounds
+			if (i >= ChildNodes.Num())
+			{
+				i = ChildNodes.Num();
+			}
 		}	
 	}
 	
