Deprecate SubobjectDataBlueprintFunctionLibrary::GetObject, replaced with GetSubobject, which properly resolves the object in the handle so that scripts don't unexpectedly begin dealing with components on the blueprint or CDO instead of objects on the instance the handle was created from
[RN] The object associated with an FSubobjectData should now be retreived by GetAssociatedObject rather than GetObject, this new function will properly resolve the object using the root context, rather than potentially returning a template object

diff --git a/Engine/Source/Editor/SubobjectDataInterface/Private/SubobjectDataBlueprintFunctionLibrary.cpp b/Engine/Source/Editor/SubobjectDataInterface/Private/SubobjectDataBlueprintFunctionLibrary.cpp
--- a/Engine/Source/Editor/SubobjectDataInterface/Private/SubobjectDataBlueprintFunctionLibrary.cpp
+++ b/Engine/Source/Editor/SubobjectDataInterface/Private/SubobjectDataBlueprintFunctionLibrary.cpp
@@ -14,6 +14,47 @@ void USubobjectDataBlueprintFunctionLibrary::GetData(const FSubobjectDataHandle&
 	}
 }
 
+const UObject* USubobjectDataBlueprintFunctionLibrary::GetObject(const FSubobjectData& Data, bool bEvenIfPendingKill)
+{
+	return Data.GetObject(bEvenIfPendingKill);
+}
+
+const UObject* USubobjectDataBlueprintFunctionLibrary::GetAssociatedObject(const FSubobjectData& Data)
+{
+	if(!Data.IsValid())
+	{
+		return nullptr;
+	}
+
+	FSubobjectDataHandle Root = Data.GetRootSubobject();
+	if(Root != FSubobjectDataHandle())
+	{
+		// resolve our object within the context of our root:
+		const UObject* Object = Root.GetData()->GetObject();
+		if(const UBlueprint* BP = Cast<UBlueprint>(Object))
+		{
+			return Data.GetObjectForBlueprint(const_cast<UBlueprint*>(BP));
+		}
+		else if(const AActor* ActorContext = Cast<AActor>(Object))
+		{
+			return Data.FindComponentInstanceInActor(ActorContext);
+		}
+	}
+
+	// we couldn't figure out the root object, hopefully the object 
+	// ptr on the subobject data is not something silly:
+	const UObject* Object = Data.GetObject();
+	if(const UBlueprint* BP = Cast<UBlueprint>(Object))
+	{
+		return BP->GeneratedClass ? BP->GeneratedClass->ClassDefaultObject : nullptr;
+	}
+	else
+	{
+		return Object;
+	}
+
+}
+
 const UObject* USubobjectDataBlueprintFunctionLibrary::GetObjectForBlueprint(const FSubobjectData& Data, UBlueprint* Blueprint) 
 {
 	return Data.GetObjectForBlueprint(Blueprint); 
diff --git a/Engine/Source/Editor/SubobjectDataInterface/Public/SubobjectDataBlueprintFunctionLibrary.h b/Engine/Source/Editor/SubobjectDataInterface/Public/SubobjectDataBlueprintFunctionLibrary.h
--- a/Engine/Source/Editor/SubobjectDataInterface/Public/SubobjectDataBlueprintFunctionLibrary.h
+++ b/Engine/Source/Editor/SubobjectDataInterface/Public/SubobjectDataBlueprintFunctionLibrary.h
@@ -90,8 +90,11 @@ public:
 	UFUNCTION(BlueprintCallable, BlueprintPure, Category="Subobject Data")
 	static bool CanRename(const FSubobjectData& Data) { return Data.CanRename(); }
 	
+	UFUNCTION(BlueprintCallable, BlueprintPure, Category="Subobject Data", meta = (BlueprintInternalUseOnly = "true", DeprecatedFunction="Note", DeprecationMessage = "Use GetAssociatedObject - it will not return template objects for instances"))
+	static const UObject* GetObject(const FSubobjectData& Data, bool bEvenIfPendingKill = false);
+
 	UFUNCTION(BlueprintCallable, BlueprintPure, Category="Subobject Data")
-	static const UObject* GetObject(const FSubobjectData& Data, bool bEvenIfPendingKill = false) { return Data.GetObject(bEvenIfPendingKill); }
+	static const UObject* GetAssociatedObject(const FSubobjectData& Data);
 
 	UFUNCTION(BlueprintCallable, Category="Subobject Data")
 	static const UObject* GetObjectForBlueprint(const FSubobjectData& Data, UBlueprint* Blueprint);
