Remove use of BlueprintInternal as it prevents python from using this function and therefore properly being allowed to adjust to deprecation

#rnx

diff --git a/Engine/Source/Editor/SubobjectDataInterface/Public/SubobjectDataBlueprintFunctionLibrary.h b/Engine/Source/Editor/SubobjectDataInterface/Public/SubobjectDataBlueprintFunctionLibrary.h
--- a/Engine/Source/Editor/SubobjectDataInterface/Public/SubobjectDataBlueprintFunctionLibrary.h
+++ b/Engine/Source/Editor/SubobjectDataInterface/Public/SubobjectDataBlueprintFunctionLibrary.h
@@ -90,7 +90,7 @@ public:
 	UFUNCTION(BlueprintCallable, BlueprintPure, Category="Subobject Data")
 	static bool CanRename(const FSubobjectData& Data) { return Data.CanRename(); }
 	
-	UFUNCTION(BlueprintCallable, BlueprintPure, Category="Subobject Data", meta = (BlueprintInternalUseOnly = "true", DeprecatedFunction="Note", DeprecationMessage = "Use GetAssociatedObject - it will not return template objects for instances"))
+	UFUNCTION(BlueprintCallable, BlueprintPure, Category="Subobject Data", meta = (DeprecatedFunction="Note", DeprecationMessage = "Use GetAssociatedObject - it will not return template objects for instances"))
 	static const UObject* GetObject(const FSubobjectData& Data, bool bEvenIfPendingKill = false);
 
 	UFUNCTION(BlueprintCallable, BlueprintPure, Category="Subobject Data")
