Fall back to the SubobjectHandle's associated object when FindComponentInstanceInActor fails to find an object, which it will for SCS defined nodes associated with a CDO

#rnx

diff --git a/Engine/Source/Editor/SubobjectDataInterface/Private/SubobjectDataBlueprintFunctionLibrary.cpp b/Engine/Source/Editor/SubobjectDataInterface/Private/SubobjectDataBlueprintFunctionLibrary.cpp
--- a/Engine/Source/Editor/SubobjectDataInterface/Private/SubobjectDataBlueprintFunctionLibrary.cpp
+++ b/Engine/Source/Editor/SubobjectDataInterface/Private/SubobjectDataBlueprintFunctionLibrary.cpp
@@ -37,7 +37,10 @@ const UObject* USubobjectDataBlueprintFunctionLibrary::GetAssociatedObject(const
 		}
 		else if(const AActor* ActorContext = Cast<AActor>(Object))
 		{
-			return Data.FindComponentInstanceInActor(ActorContext);
+			if(const UObject* Component = Data.FindComponentInstanceInActor(ActorContext))
+			{
+				return Component;
+			}
 		}
 	}
 
