DefaultSceneRoot component now resets its name when "deleted".

diff --git a/Engine/Source/Editor/SubobjectDataInterface/Private/SubobjectDataSubsystem.cpp b/Engine/Source/Editor/SubobjectDataInterface/Private/SubobjectDataSubsystem.cpp
--- a/Engine/Source/Editor/SubobjectDataInterface/Private/SubobjectDataSubsystem.cpp
+++ b/Engine/Source/Editor/SubobjectDataInterface/Private/SubobjectDataSubsystem.cpp
@@ -1242,6 +1242,13 @@ int32 USubobjectDataSubsystem::DeleteSubobjects(const FSubobjectDataHandle& Cont
 						{
 							UE_LOG(LogSubobjectSubsystem, Warning, TEXT("Cannot remove subobject '%s' because it is the default scene root!"), *Data->GetDisplayString());
 						}
+
+						if (Data->IsDefaultSceneRoot())
+						{
+							// The DefaultSceneRoot is never deleted, but just hidden instead.
+							// We reset its name, so that it can be reused by other Components.
+							SCS_Node->SetVariableName(USceneComponent::GetDefaultSceneRootVariableName());
+						}
 					}
 				}
 			}
