Add null checks to avoid crashes from missing component templates (e.g. editor only components in cooked parent BPs)
#rnx

diff --git a/Engine/Source/Editor/SubobjectDataInterface/Private/SubobjectData.cpp b/Engine/Source/Editor/SubobjectDataInterface/Private/SubobjectData.cpp
--- a/Engine/Source/Editor/SubobjectDataInterface/Private/SubobjectData.cpp
+++ b/Engine/Source/Editor/SubobjectDataInterface/Private/SubobjectData.cpp
@@ -31,7 +31,7 @@ FSubobjectData::FSubobjectData(UObject* ContextObject, const FSubobjectDataHandl
 	// Check for a BP added CAC
 	if (USCS_Node* SCS = Cast<USCS_Node>(ContextObject))
 	{
-		if (SCS->ComponentTemplate->IsA<UChildActorComponent>())
+		if (SCS->ComponentTemplate && SCS->ComponentTemplate->IsA<UChildActorComponent>())
 		{
 			bIsChildActor = true;
 		}
diff --git a/Engine/Source/Editor/SubobjectDataInterface/Private/SubobjectDataSubsystem.cpp b/Engine/Source/Editor/SubobjectDataInterface/Private/SubobjectDataSubsystem.cpp
--- a/Engine/Source/Editor/SubobjectDataInterface/Private/SubobjectDataSubsystem.cpp
+++ b/Engine/Source/Editor/SubobjectDataInterface/Private/SubobjectDataSubsystem.cpp
@@ -300,7 +300,7 @@ void USubobjectDataSubsystem::GatherSubobjectData(UObject* Context, TArray<FSubo
 					{
 						FSubobjectDataHandle& ParentHandle = RootActorHandle;
 						// Only a valid scene component can be used as the parent handle for a subobject, otherwise it should be attached to the root actor
-						if (RootComponentHandle.IsValid() && SCS_Node->ComponentTemplate->IsA<USceneComponent>())
+						if (RootComponentHandle.IsValid() && SCS_Node->ComponentTemplate && SCS_Node->ComponentTemplate->IsA<USceneComponent>())
 						{
 							ParentHandle = RootComponentHandle;
 						}
