Blueprint Editor: The tree view should now remain sorted alphabetically with scene components at the top.

[RN]

diff --git a/Engine/Source/Editor/SubobjectDataInterface/Private/SubobjectDataSubsystem.cpp b/Engine/Source/Editor/SubobjectDataInterface/Private/SubobjectDataSubsystem.cpp
--- a/Engine/Source/Editor/SubobjectDataInterface/Private/SubobjectDataSubsystem.cpp
+++ b/Engine/Source/Editor/SubobjectDataInterface/Private/SubobjectDataSubsystem.cpp
@@ -329,10 +329,30 @@ void USubobjectDataSubsystem::GatherSubobjectData(UObject* Context, TArray<FSubo
 			}
 		}
 
-		// Sort subobjects by type (always put scene components first in the tree)
+		// Sort subobjects by type (always put scene components first in the tree), then 
 		OutArray.Sort([](const FSubobjectDataHandle& A, const FSubobjectDataHandle& B)
 		{
-			return A.GetData()->IsActor() || (A.GetData()->IsSceneComponent() && !B.GetData()->IsActor());
+			if (A.GetData()->IsSceneComponent())
+			{
+				if (B.GetData()->IsSceneComponent())
+				{
+					return A.GetData()->GetDisplayName().CompareTo(B.GetData()->GetDisplayName()) < 0;
+				}
+
+				return !B.GetData()->IsActor();
+			}
+
+			if (A.GetData()->IsActor())
+			{
+				if (B.GetData()->IsActor())
+				{
+					return A.GetData()->GetDisplayName().CompareTo(B.GetData()->GetDisplayName()) < 0;
+				}
+
+				return true;
+			}
+
+			return A.GetData()->GetDisplayName().CompareTo(B.GetData()->GetDisplayName()) < 0;
 		});
 	}
 	// Otherwise, this is an actor instance in a level
