CL adds extension points into the SubobjectDataInterface and extends the existing UI extension point for the SubobjectEditor.  The purpose of this CL is to start retrofitting this system to not be dependent on Actors and Components and instead use other types of objects (specifically targeting verse entity + component).

The intention of this CL is that all legacy behaviour will be preserved if no extension points are registered.  If an extension adapter is registered and it doesn't pass the filter that was registered with it, then the system will fallback to legacy behaviour.

Where to start:

# FSubobjectData
This is an already reasonably type agnostic proxy for an element within the SuobbjectEditor tree.  There are two major changes.

##  ContextData FInstancedStruct
- This is a replacement for ContextObject and can either return a UObject (legacy) or a RowHandle (new!).  Note: Whenever possible, the existing ContextObject is still set.

##  SubobjectDataAdapter
- This hooks the SubobjectData member functions to allow injection of behavior which is activated based on whether the ContextData passes some specific filters.
- Wanted the filters to be TEDS specific, but opted to make it a function since I could see us implementing the existing legacy behaviour using the exact same mechanisms being proposed to extend the system.

# SubobjectDataSubsystem
This is the top level object that generates SubobjectData instances.
- Added an extension point to GatherSubobjectData which does the initial generation of the tree
- There is missing extension here though to handle newly added subobjects and reordered subobjects.  This will come later.

# ISCSEditorUICustomization
Despite the name, this is the customization for the UX of SubobjectEditors
- GetIconBrush provides a way for the customization to chose the icon for the SubobjectData row
- GetControlsWidget - provides a way for the customization to inject widgets that are right-justified.  This would replace something like "Edit in C++" for non-native component types.

# SActorDetails
Despite the name, this is the details panel customization for the LevelEditor.  It generates a SubobjectInstanceEditor (derivation of SubobjectEditor).
There are only two small changes here, despite lines of code.
1. Renamed GetActorContextAsObject to GetObjectContext to handle non-Actor based UObjects in the LevelEditor details panel. Removed the internal AActor cast
2. Refactored OnSubobjectEditorTreeViewSelectionChanged to handle non-Actor based UObjects set in the Subobject Editor
- Very importantly, TEDS is used to find the appropriate TypedElementHandle given the UObject passed to it via the TedsTypedElementColumn.  In theory, the original loop could be replaced by the one added here but left the legacy loop in place as risk mitigation for now.

diff --git a/Engine/Source/Editor/Kismet/Public/ISCSEditorUICustomization.h b/Engine/Source/Editor/Kismet/Public/ISCSEditorUICustomization.h
--- a/Engine/Source/Editor/Kismet/Public/ISCSEditorUICustomization.h
+++ b/Engine/Source/Editor/Kismet/Public/ISCSEditorUICustomization.h
@@ -8,6 +8,7 @@
 #include "Containers/ArrayView.h"
 
 class UActorComponent;
+struct FSubobjectData;
 
 /** SCSEditor UI customization */
 class ISCSEditorUICustomization
@@ -42,6 +43,23 @@ public:
 	/** @return Whether to hide the "Edit Blueprint" and "Blueprint/Add Script" buttons */
 	virtual bool HideBlueprintButtons(TArrayView<UObject*> Context) const { return false; }
 
+	/** @return The icon that should be used for the given SubobjectData
+	 * 
+	 * Implementations that return nullptr will lead to existing built-in functionality of the
+	 * subobject editor being used to determine the Icon for the given FSubobjectData 
+	 */
+	virtual const FSlateBrush* GetIconBrush(const FSubobjectData&) const { return nullptr; }
+
+	/** Allows customizer to inject a right-side justified widget on the subobject editor
+	 * Typically these widgets would convey actions that can be taken and may optionally provide
+	 * status information through their visualization.
+	 * As an example, a subobject that is a non-native ActorComponent may provide a button via these widgets
+	 * to open a code editor to the C++ version of that type.
+	 * Returning nullptr will cause built-in behaviour to be used, indicating that the UICustomization does not customize this
+	 * UI element. To override existing behaviour and actually show nothing, return a single null widget
+	 */
+	virtual TSharedPtr<SWidget> GetControlsWidget(const FSubobjectData&) const { return {}; }
+
 	/** 
 	 * @return An override for the default ChildActorComponentTreeViewVisualizationMode from the project settings, if different from UseDefault.
 	 * @note Setting an override also forces child actor tree view expansion to be enabled.
diff --git a/Engine/Source/Editor/LevelEditor/Private/ActorDetailsSCSEditorUICustomization.cpp b/Engine/Source/Editor/LevelEditor/Private/ActorDetailsSCSEditorUICustomization.cpp
--- a/Engine/Source/Editor/LevelEditor/Private/ActorDetailsSCSEditorUICustomization.cpp
+++ b/Engine/Source/Editor/LevelEditor/Private/ActorDetailsSCSEditorUICustomization.cpp
@@ -64,6 +64,30 @@ bool FActorDetailsSCSEditorUICustomization::HideBlueprintButtons(TArrayView<UObj
 	return false;
 }
 
+const FSlateBrush* FActorDetailsSCSEditorUICustomization::GetIconBrush(const FSubobjectData& SubobjectData) const
+{
+	for (const TSharedPtr<ISCSEditorUICustomization>& Customization : Customizations)
+	{
+		if (const FSlateBrush* IconBrush = Customization->GetIconBrush(SubobjectData))
+		{
+			return IconBrush;
+		}
+	}
+	return ISCSEditorUICustomization::GetIconBrush(SubobjectData);
+}
+
+TSharedPtr<SWidget> FActorDetailsSCSEditorUICustomization::GetControlsWidget(const FSubobjectData& SubobjectData) const
+{
+	for (const TSharedPtr<ISCSEditorUICustomization>& Customization : Customizations)
+	{
+		if (TSharedPtr<SWidget> ControlsWidget = Customization->GetControlsWidget(SubobjectData))
+		{
+			return ControlsWidget;
+		}
+	}
+	return ISCSEditorUICustomization::GetControlsWidget(SubobjectData);
+}
+
 EChildActorComponentTreeViewVisualizationMode FActorDetailsSCSEditorUICustomization::GetChildActorVisualizationMode() const
 {
 	for (const TSharedPtr<ISCSEditorUICustomization>& Customization : Customizations)
diff --git a/Engine/Source/Editor/LevelEditor/Private/ActorDetailsSCSEditorUICustomization.h b/Engine/Source/Editor/LevelEditor/Private/ActorDetailsSCSEditorUICustomization.h
--- a/Engine/Source/Editor/LevelEditor/Private/ActorDetailsSCSEditorUICustomization.h
+++ b/Engine/Source/Editor/LevelEditor/Private/ActorDetailsSCSEditorUICustomization.h
@@ -21,6 +21,9 @@ public:
 
 	virtual bool HideBlueprintButtons(TArrayView<UObject*> Context) const override;
 
+	virtual const FSlateBrush* GetIconBrush(const FSubobjectData&) const override;
+	virtual TSharedPtr<SWidget> GetControlsWidget(const FSubobjectData&) const override;
+
 	virtual EChildActorComponentTreeViewVisualizationMode GetChildActorVisualizationMode() const override;
 
 	virtual TSubclassOf<UActorComponent> GetComponentTypeFilter(TArrayView<UObject*> Context) const override;
@@ -31,6 +34,7 @@ public:
 private:
 	FActorDetailsSCSEditorUICustomization() {}
 
+private:
 	static TSharedPtr<FActorDetailsSCSEditorUICustomization> Instance;
 	TArray<TSharedPtr<ISCSEditorUICustomization>> Customizations;
 };
\ No newline at end of file
diff --git a/Engine/Source/Editor/LevelEditor/Private/SActorDetails.cpp b/Engine/Source/Editor/LevelEditor/Private/SActorDetails.cpp
--- a/Engine/Source/Editor/LevelEditor/Private/SActorDetails.cpp
+++ b/Engine/Source/Editor/LevelEditor/Private/SActorDetails.cpp
@@ -6,6 +6,7 @@
 #include "Components/SceneComponent.h"
 #include "DetailsViewObjectFilter.h"
 #include "Editor.h"
+#include "EditorLevelUtils.h"
 #include "Editor/UnrealEdEngine.h"
 #include "Styling/AppStyle.h"
 #include "Elements/Framework/EngineElementsLibrary.h"
@@ -34,6 +35,8 @@
 #include "SubobjectData.h"
 #include "SubobjectDataSubsystem.h"
 #include "UnrealEdGlobals.h"
+#include "DataStorage/Features.h"
+#include "Elements/Columns/TypedElementHandleColumn.h"
 #include "Widgets/Docking/SDockTab.h"
 #include "Widgets/Images/SImage.h"
 #include "Widgets/Input/SButton.h"
@@ -272,7 +275,7 @@ void SActorDetails::Construct(const FArguments& InArgs, UTypedElementSelectionSe
 	FModuleManager::LoadModuleChecked<FSubobjectEditorModule>("SubobjectEditor");
 	
 	SubobjectEditor = SNew(SSubobjectInstanceEditor)
-		.ObjectContext(this, &SActorDetails::GetActorContextAsObject)
+		.ObjectContext(this, &SActorDetails::GetObjectContext)
 		.AllowEditing(this, &SActorDetails::GetAllowComponentTreeEditing)
 		.OnSelectionUpdated(this, &SActorDetails::OnSubobjectEditorTreeViewSelectionChanged)
 		.OnItemDoubleClicked(this, &SActorDetails::OnSubobjectEditorTreeViewItemDoubleClicked);
@@ -518,7 +521,7 @@ void SActorDetails::RefreshSubobjectTreeElements(TArrayView<const FSubobjectEdit
 	}
 
 	// Does the Subobject tree have components selected?
-	TArray<const UActorComponent*> Components;
+	TArray<FTypedElementHandle> Components;
 	if (const AActor* Actor = GetActorContext())
 	{
 		for (const FSubobjectEditorTreeNodePtrType& SelectedNode : InSelectedNodes)
@@ -530,15 +533,52 @@ void SActorDetails::RefreshSubobjectTreeElements(TArrayView<const FSubobjectEdit
 					if(Data->IsRootActor())
 					{
 						// If the actor node is selected then we ignore the component selection
-                        Components.Reset();
-                        break;
+						Components.Reset();
+						break;
 					}
 					
 					if (Data->IsComponent())
 					{
 						if (const UActorComponent* Component = Data->FindComponentInstanceInActor(Actor))
 						{
-							Components.Add(Component);
+							if (FTypedElementHandle ComponentElementHandle = UEngineElementsLibrary::AcquireEditorComponentElementHandle(Component))
+							{
+								Components.Add(ComponentElementHandle);
+							}
+						}
+					}
+				}
+			}
+		}
+	}
+	else
+	{
+		using namespace UE::Editor::DataStorage;
+		const ICoreProvider* CoreProvider = GetDataStorageFeature<ICoreProvider>(StorageFeatureName);
+		for (const FSubobjectEditorTreeNodePtrType& SelectedNode : InSelectedNodes)
+		{
+			if (SelectedNode)
+			{
+				if(const FSubobjectData* Data = SelectedNode->GetDataSource())
+				{
+					if(Data->IsRootActor()) // Not an actor, using IsRootActor to mean IsRootNode
+					{
+						// If the root node is selected then we ignore the component selection
+						Components.Reset();
+						break;
+					}
+					
+					if (Data->IsComponent())
+					{
+						const FSubobjectDataSubsystemContextData_TedsRow* TedsRowAdapter = Data->GetContextData().GetPtr<FSubobjectDataSubsystemContextData_TedsRow>();
+						if (TedsRowAdapter)
+						{
+							UE::Editor::DataStorage::RowHandle ComponentRowHandle = TedsRowAdapter->RowHandle;
+							if (const FTedsTypedElementColumn* TypedElementHandleColumn = CoreProvider->GetColumn<FTedsTypedElementColumn>(ComponentRowHandle))
+							{
+								FTypedElementHandle ComponentElementHandle = TypedElementHandleColumn->Handle;
+								Components.Add(TypedElementHandleColumn->Handle);
+							}
 						}
 					}
 				}
@@ -550,16 +590,13 @@ void SActorDetails::RefreshSubobjectTreeElements(TArrayView<const FSubobjectEdit
 	if (Components.Num() > 0)
 	{
 		UTypedElementRegistry* Registry = UTypedElementRegistry::GetInstance();
-		for (const UActorComponent* Component : Components)
+		for (const FTypedElementHandle& ComponentElementHandle : Components)
 		{
-			if (FTypedElementHandle ComponentElementHandle = UEngineElementsLibrary::AcquireEditorComponentElementHandle(Component))
+			if (TTypedElement<ITypedElementDetailsInterface> ComponentDetailsHandle = Registry->GetElement<ITypedElementDetailsInterface>(ComponentElementHandle))
 			{
-				if (TTypedElement<ITypedElementDetailsInterface> ComponentDetailsHandle = Registry->GetElement<ITypedElementDetailsInterface>(ComponentElementHandle))
+				if (TUniquePtr<ITypedElementDetailsObject> ElementDetailsObject = ComponentDetailsHandle.GetDetailsObject())
 				{
-					if (TUniquePtr<ITypedElementDetailsObject> ElementDetailsObject = ComponentDetailsHandle.GetDetailsObject())
-					{
-						SubobjectTreeElements.Add(MoveTemp(ElementDetailsObject));
-					}
+					SubobjectTreeElements.Add(MoveTemp(ElementDetailsObject));
 				}
 			}
 		}
@@ -665,9 +702,11 @@ AActor* SActorDetails::GetActorContext() const
 		: nullptr;
 }
 
-UObject* SActorDetails::GetActorContextAsObject() const
+UObject* SActorDetails::GetObjectContext() const
 {
-	return GetActorContext();
+	return TopLevelElements.Num() == 1
+		? TopLevelElements[0]->GetObject()
+		: nullptr;
 }
 
 void SActorDetails::OnSubobjectEditorTreeViewSelectionChanged(const TArray<FSubobjectEditorTreeNodePtrType>& SelectedNodes)
@@ -683,78 +722,77 @@ void SActorDetails::OnSubobjectEditorTreeViewSelectionChanged(const TArray<FSubo
 		// Don't respond to de-selecting everything...
 		return;
 	}
-
-	AActor* ActorContext = GetActorContext();
-	if (!ActorContext)
-	{
-		// The Subobject editor requires an actor context...
-		return;
-	}
-
-	if (SelectedNodes.Num() > 1 && SelectedBPComponentBlueprint.IsValid())
-	{
-		// Remove the compilation delegate if we are no longer displaying the full details for a single blueprint component.
-		RemoveBPComponentCompileEventDelegate();
-	}
-	else if (SelectedNodes.Num() == 1 && SelectedNodes[0]->IsComponentNode())
+	
+	if (AActor* ActorContext = GetActorContext())
 	{
-		// Add delegate to monitor blueprint component compilation if we have a full details view ( i.e. single selection )
-		FSubobjectData* SelectedData = SelectedNodes[0]->GetDataSource();
-		if (UActorComponent* Component = const_cast<UActorComponent*>(SelectedData->FindComponentInstanceInActor(ActorContext)))
+		if (SelectedNodes.Num() > 1 && SelectedBPComponentBlueprint.IsValid())
+		{
+			// Remove the compilation delegate if we are no longer displaying the full details for a single blueprint component.
+			RemoveBPComponentCompileEventDelegate();
+		}
+		else if (SelectedNodes.Num() == 1 && SelectedNodes[0]->IsComponentNode())
 		{
-			if (UBlueprintGeneratedClass* ComponentBPGC = Cast<UBlueprintGeneratedClass>(Component->GetClass()))
+			// Add delegate to monitor blueprint component compilation if we have a full details view ( i.e. single selection )
+			FSubobjectData* SelectedData = SelectedNodes[0]->GetDataSource();
+			if (UActorComponent* Component = const_cast<UActorComponent*>(SelectedData->FindComponentInstanceInActor(ActorContext)))
 			{
-				if (UBlueprint* ComponentBlueprint = Cast<UBlueprint>(ComponentBPGC->ClassGeneratedBy))
+				if (UBlueprintGeneratedClass* ComponentBPGC = Cast<UBlueprintGeneratedClass>(Component->GetClass()))
 				{
-					AddBPComponentCompileEventDelegate(ComponentBlueprint);
+					if (UBlueprint* ComponentBlueprint = Cast<UBlueprint>(ComponentBPGC->ClassGeneratedBy))
+					{
+						AddBPComponentCompileEventDelegate(ComponentBlueprint);
+					}
 				}
 			}
 		}
-	}
-	
-	// We only actually update the editor selection state if we're not locked
-	if (!DetailsView->IsLocked())
-	{
-		TArray<FTypedElementHandle> NewEditorSelection;
-		NewEditorSelection.Add(UEngineElementsLibrary::AcquireEditorActorElementHandle(ActorContext));
-
-		for (const FSubobjectEditorTreeNodePtrType& SelectedNode : SelectedNodes)
+		
+		// We only actually update the editor selection state if we're not locked
+		if (!DetailsView->IsLocked())
 		{
-			if (SelectedNode)
+			TArray<FTypedElementHandle> NewEditorSelection;
+			NewEditorSelection.Add(UEngineElementsLibrary::AcquireEditorActorElementHandle(ActorContext));
+
+			for (const FSubobjectEditorTreeNodePtrType& SelectedNode : SelectedNodes)
 			{
-				if(FSubobjectData* Data = SelectedNode->GetDataSource())
+				if (SelectedNode)
 				{
-					if(Data->IsRootActor())
-					{
-						// If the actor node is selected then we ignore the component selection
-						NewEditorSelection.Reset();
-						NewEditorSelection.Add(UEngineElementsLibrary::AcquireEditorActorElementHandle(ActorContext));
-						break;
-					}
-				
-					if (Data->IsComponent())
+					if(FSubobjectData* Data = SelectedNode->GetDataSource())
 					{
-						if (const UActorComponent* Component = Data->FindComponentInstanceInActor(ActorContext))
+						if(Data->IsRootActor())
 						{
-							NewEditorSelection.Add(UEngineElementsLibrary::AcquireEditorComponentElementHandle(Component));
+							// If the actor node is selected then we ignore the component selection
+							NewEditorSelection.Reset();
+							NewEditorSelection.Add(UEngineElementsLibrary::AcquireEditorActorElementHandle(ActorContext));
+							break;
+						}
+					
+						if (Data->IsComponent())
+						{
+							if (const UActorComponent* Component = Data->FindComponentInstanceInActor(ActorContext))
+							{
+								NewEditorSelection.Add(UEngineElementsLibrary::AcquireEditorComponentElementHandle(Component));
+							}
 						}
 					}
 				}
 			}
-		}
 
-		// Note: this transaction should not take place if we are in the middle of executing an undo or redo because it would clear the top of the transaction stack.
-		const bool bShouldActuallyTransact = !GIsTransacting;
-		const FScopedTransaction Transaction(NSLOCTEXT("UnrealEd", "ClickingOnComponentInTree", "Clicking on Component (tree view)"), bShouldActuallyTransact);
+			// Note: this transaction should not take place if we are in the middle of executing an undo or redo because it would clear the top of the transaction stack.
+			const bool bShouldActuallyTransact = !GIsTransacting;
+			const FScopedTransaction Transaction(NSLOCTEXT("UnrealEd", "ClickingOnComponentInTree", "Clicking on Component (tree view)"), bShouldActuallyTransact);
 
-		// Enable the selection guard to prevent OnEditorSelectionChanged() from altering the contents of the SubobjectTreeWidget
-		TGuardValue<bool> SelectionGuard(bSelectionGuard, true);
-		SelectionSet->SetSelection(NewEditorSelection, FTypedElementSelectionOptions());
-		SelectionSet->NotifyPendingChanges(); // Fire while still under the selection guard
+			// Enable the selection guard to prevent OnEditorSelectionChanged() from altering the contents of the SubobjectTreeWidget
+			TGuardValue<bool> SelectionGuard(bSelectionGuard, true);
+			SelectionSet->SetSelection(NewEditorSelection, FTypedElementSelectionOptions());
+			SelectionSet->NotifyPendingChanges(); // Fire while still under the selection guard
+		}
+		// Update the underlying details view
+		RefreshSubobjectTreeElements(SelectedNodes, /*bForceRefresh*/false, DetailsView->IsLocked());
+	}
+	else if (UObject* ObjectContext = GetObjectContext())
+	{
+		RefreshSubobjectTreeElements(SelectedNodes, /*bForceRefresh*/false, DetailsView->IsLocked());
 	}
-
-	// Update the underlying details view
-	RefreshSubobjectTreeElements(SelectedNodes, /*bForceRefresh*/false, DetailsView->IsLocked());
 }
 
 void SActorDetails::OnSubobjectEditorTreeViewItemDoubleClicked(const FSubobjectEditorTreeNodePtrType ClickedNode)
@@ -888,12 +926,13 @@ EVisibility SActorDetails::GetComponentEditorVisibility() const
 	// see if we need to hide the editor due to the current object display 
 	const bool bHideEditorFromDetailsView = (DisplayManager.IsValid() &&
 							                 DisplayManager->ShouldHideComponentEditor() );
-	return GetActorContext() && !bHideEditorFromDetailsView  ? EVisibility::Visible : EVisibility::Collapsed;
+
+	return GetObjectContext() && !bHideEditorFromDetailsView  ? EVisibility::Visible : EVisibility::Collapsed;
 }
 
 EVisibility SActorDetails::GetComponentEditorButtonsVisibility() const
 {
-	return GetActorContext() ? EVisibility::Visible : EVisibility::Collapsed;
+	return GetObjectContext() ? EVisibility::Visible : EVisibility::Collapsed;
 }
 
 EVisibility SActorDetails::GetUCSComponentWarningVisibility() const
diff --git a/Engine/Source/Editor/LevelEditor/Private/SActorDetails.h b/Engine/Source/Editor/LevelEditor/Private/SActorDetails.h
--- a/Engine/Source/Editor/LevelEditor/Private/SActorDetails.h
+++ b/Engine/Source/Editor/LevelEditor/Private/SActorDetails.h
@@ -33,7 +33,7 @@ namespace UE::LevelEditor::Private
 }
 
 /**
- * Wraps a details panel customized for viewing actors
+ * Wraps a details panel customized for the level editor
  */
 class SActorDetails : public SCompoundWidget, public FEditorUndoClient, public FGCObject
 {
@@ -86,7 +86,7 @@ private:
 	void SetElementDetailsObjects(TArrayView<const TUniquePtr<ITypedElementDetailsObject>> InElementDetailsObjects, const bool bForceRefresh, const bool bOverrideLock, TArrayView<const TTypedElement<ITypedElementDetailsInterface>>* InDetailsElementsPtr = nullptr);
 
 	AActor* GetActorContext() const;
-	UObject* GetActorContextAsObject() const;
+	UObject* GetObjectContext() const;
 	bool GetAllowComponentTreeEditing() const;
 
 	/**
diff --git a/Engine/Source/Editor/SubobjectDataInterface/Private/SubobjectData.cpp b/Engine/Source/Editor/SubobjectDataInterface/Private/SubobjectData.cpp
--- a/Engine/Source/Editor/SubobjectDataInterface/Private/SubobjectData.cpp
+++ b/Engine/Source/Editor/SubobjectDataInterface/Private/SubobjectData.cpp
@@ -8,6 +8,11 @@
 #include "Kismet2/ComponentEditorUtils.h"
 #include "Kismet2/BlueprintEditorUtils.h"
 #include "ComponentAssetBroker.h"			// FComponentAssetBrokerage
+#include "SubobjectDataSubsystem.h"
+#include "DataStorage/Features.h"
+#include "Elements/Columns/TypedElementLabelColumns.h"
+#include "Elements/Columns/TypedElementTypeInfoColumns.h"
+#include "Elements/Interfaces/TypedElementDataStorageInterface.h"
 #include "Engine/World.h"					// Finding the instance of a component on an actor
 #include "Engine/InheritableComponentHandler.h"
 #include "Engine/SCS_Node.h"		// #TODO_BH  We need to remove this when the actual subobject refactor happens
@@ -16,6 +21,14 @@
 
 #define LOCTEXT_NAMESPACE "SubobjectData"
 
+PRAGMA_DISABLE_DEPRECATION_WARNINGS
+#define REDIRECT_TO_ADAPTER(FunctionName, ...) \
+	if (const FSubobjectDataAdapterBase* Adapter = USubobjectDataSubsystem::Get()->FindDataAdapter(*this))				\
+	{																													\
+		return Adapter->FunctionName(*this, ##__VA_ARGS__);																\
+	}
+PRAGMA_ENABLE_DEPRECATION_WARNINGS
+
 FSubobjectData::FSubobjectData()
 	: WeakObjectPtr(nullptr)
 	, ParentObjectHandle(FSubobjectDataHandle::InvalidHandle)
@@ -30,6 +43,9 @@ FSubobjectData::FSubobjectData(UObject* ContextObject, const FSubobjectDataHandl
 	, ParentObjectHandle(InParentHandle)
 	, SCSNodePtr(nullptr)
 {
+	FSubobjectDataSubsystemContextData_SingleUObjectContextObject SingleUObject(ContextObject);
+	ContextData.InitializeAsScriptStruct(FSubobjectDataSubsystemContextData_SingleUObjectContextObject::StaticStruct(), reinterpret_cast<const uint8*>(&SingleUObject));
+	
 	// Check for a BP added CAC
 	if (USCS_Node* SCS = Cast<USCS_Node>(ContextObject))
 	{
@@ -64,8 +80,22 @@ FSubobjectData::FSubobjectData(UObject* ContextObject, const FSubobjectDataHandl
 	AttemptToSetSCSNode();
 }
 
+void FSubobjectData::Initialize(TInstancedStruct<FSubobjectDataSubsystemContextDataBase> InContextData, const FSubobjectDataHandle& ParentHandle)
+{
+	// Avoid double initializations
+	if (ContextData.IsValid())
+	{
+		return;
+	}
+
+	ContextData = MoveTemp(InContextData);
+	ParentObjectHandle = ParentHandle;
+}
+
 bool FSubobjectData::CanEdit() const
 {
+	REDIRECT_TO_ADAPTER(CanEdit);
+	
 	if (bIsInheritedSubobject)
 	{
 		if (IsComponent())
@@ -106,6 +136,8 @@ bool FSubobjectData::CanEdit() const
 
 bool FSubobjectData::CanDelete() const
 {
+	REDIRECT_TO_ADAPTER(CanDelete);
+
 	if (bIsInheritedSubobject)
 	{
 		if (IsInheritedComponent() || (IsDefaultSceneRoot() && SceneRootHasDefaultName()) || (GetSCSNode() != nullptr && IsInstancedInheritedComponent()) || IsChildActorSubtreeObject())
@@ -149,11 +181,15 @@ bool FSubobjectData::CanDelete() const
 
 bool FSubobjectData::CanDuplicate() const
 {
+	REDIRECT_TO_ADAPTER(CanDuplicate);
+
 	return CanCopy();
 }
 
 bool FSubobjectData::CanCopy() const
 {
+	REDIRECT_TO_ADAPTER(CanCopy);
+	
 	if (IsInstancedInheritedComponent())
 	{
 		return FComponentEditorUtils::CanCopyComponent(GetComponentTemplate()) && !SceneRootHasDefaultName();
@@ -164,6 +200,8 @@ bool FSubobjectData::CanCopy() const
 
 bool FSubobjectData::CanReparent() const
 {
+	REDIRECT_TO_ADAPTER(CanReparent);
+	
 	if(IsComponent())
 	{
 		if(GetSCSNode() != nullptr)
@@ -179,6 +217,8 @@ bool FSubobjectData::CanReparent() const
 
 bool FSubobjectData::CanRename() const
 {
+	REDIRECT_TO_ADAPTER(CanRename);
+	
 	if(!IsValid())
 	{
 		return false;
@@ -203,6 +243,8 @@ bool FSubobjectData::CanRename() const
 
 const UObject* FSubobjectData::GetObjectForBlueprint(UBlueprint* Blueprint) const
 {
+	REDIRECT_TO_ADAPTER(GetObjectForBlueprint, Blueprint);
+	
 	const bool bCanEdit = CanEdit();
 	// We have to deal with the ICH (Inherited Component Handler) for components 
 	if(IsComponent() && bCanEdit && !IsNativeComponent() && IsInheritedSCSNode())
@@ -237,6 +279,8 @@ const UObject* FSubobjectData::GetObjectForBlueprint(UBlueprint* Blueprint) cons
 
 bool FSubobjectData::IsInstancedComponent() const
 {
+	REDIRECT_TO_ADAPTER(IsInstancedComponent);
+	
 	// Check the flags on the component (RF_ClassDefaultObject | RF_ArchetypeObject)
 	const UActorComponent* Component = GetComponentTemplate();
 	return Component && !Component->IsTemplate();
@@ -244,6 +288,8 @@ bool FSubobjectData::IsInstancedComponent() const
 
 UActorComponent* FSubobjectData::FindMutableComponentInstanceInActor(const AActor* InActor) const
 {
+	REDIRECT_TO_ADAPTER(FindMutableComponentInstanceInActor, InActor);
+	
 	const USCS_Node* SCS_Node = GetSCSNode();
 	const UActorComponent* ComponentTemplate = GetComponentTemplate();
 
@@ -295,6 +341,8 @@ UActorComponent* FSubobjectData::FindMutableComponentInstanceInActor(const AActo
 
 UBlueprint* FSubobjectData::GetBlueprint() const
 {
+	REDIRECT_TO_ADAPTER(GetBlueprint);
+	
 	// If this object is a BP, we can just return that
 	if (UBlueprint* BP = Cast<UBlueprint>(WeakObjectPtr.Get()))
 	{
@@ -329,6 +377,8 @@ UBlueprint* FSubobjectData::GetBlueprint() const
 
 UBlueprint* FSubobjectData::GetBlueprintBeingEdited() const
 {
+	REDIRECT_TO_ADAPTER(GetBlueprintBeingEdited);
+	
 	UObject* Object = WeakObjectPtr.Get();
 	if (!Object)
 	{
@@ -349,6 +399,8 @@ UBlueprint* FSubobjectData::GetBlueprintBeingEdited() const
 
 FString FSubobjectData::GetDisplayString(bool bShowNativeComponentNames /* = true */) const
 {
+	REDIRECT_TO_ADAPTER(GetDisplayString, bShowNativeComponentNames);
+	
 	FName VariableName = GetVariableName();
 
 	const UActorComponent* ComponentTemplate = GetComponentTemplate();
@@ -430,11 +482,15 @@ FString FSubobjectData::GetDisplayString(bool bShowNativeComponentNames /* = tru
 
 FText FSubobjectData::GetDragDropDisplayText() const
 {
+	REDIRECT_TO_ADAPTER(GetDragDropDisplayText);
+	
 	return FText::FromString(GetDisplayString());
 }
 
 FText FSubobjectData::GetDisplayNameContextModifiers(bool bShowNativeComponentNames) const
 {
+	REDIRECT_TO_ADAPTER(GetDisplayNameContextModifiers, bShowNativeComponentNames);
+	
 	if(IsActor())
 	{
 		if(IsChildActor())
@@ -494,9 +550,11 @@ FText FSubobjectData::GetDisplayNameContextModifiers(bool bShowNativeComponentNa
 
 FText FSubobjectData::GetDisplayName() const
 {
+	REDIRECT_TO_ADAPTER(GetDisplayName);
+	
 	if (bIsChildActor)
 	{
-		if (const UChildActorComponent* CAC = GetChildActorComponent())
+		if (const UChildActorComponent* CAC = GetObject<UChildActorComponent>())
 		{
 			return CAC->GetClass()->GetDisplayNameText();
 		}
@@ -528,6 +586,8 @@ FText FSubobjectData::GetDisplayName() const
 
 FName FSubobjectData::GetVariableName() const
 {
+	REDIRECT_TO_ADAPTER(GetVariableName);
+	
 	FName VariableName = NAME_None;
 
 	const USCS_Node* SCS_Node = GetSCSNode();
@@ -549,6 +609,8 @@ FName FSubobjectData::GetVariableName() const
 
 FText FSubobjectData::GetSocketName() const
 {
+	REDIRECT_TO_ADAPTER(GetSocketName);
+	
 	if (USCS_Node* SCSNode = GetSCSNode())
 	{
 		return FText::FromName(SCSNode->AttachToName);
@@ -558,6 +620,8 @@ FText FSubobjectData::GetSocketName() const
 
 FName FSubobjectData::GetSocketFName() const
 {
+	REDIRECT_TO_ADAPTER(GetSocketFName);
+
 	if (USCS_Node* SCSNode = GetSCSNode())
 	{
 		return SCSNode->AttachToName;
@@ -567,11 +631,15 @@ FName FSubobjectData::GetSocketFName() const
 
 bool FSubobjectData::HasValidSocket() const
 {
+	REDIRECT_TO_ADAPTER(HasValidSocket);
+	
 	return GetSCSNode() != nullptr;
 }
 
 void FSubobjectData::SetSocketName(FName InNewName)
 {
+	REDIRECT_TO_ADAPTER(SetSocketName, InNewName);
+
 	if(USCS_Node* SCS = GetSCSNode())
 	{
 		SCS->AttachToName = InNewName;
@@ -580,6 +648,8 @@ void FSubobjectData::SetSocketName(FName InNewName)
 
 void FSubobjectData::SetupAttachment(FName SocketName, const FSubobjectDataHandle& AttachParentHandle)
 {
+	REDIRECT_TO_ADAPTER(SetupAttachment, SocketName, AttachParentHandle);
+
 	USceneComponent* SceneComponentTemplate = Cast<USceneComponent>(GetMutableComponentTemplate());
 	if (!SceneComponentTemplate)
 	{
@@ -637,6 +707,8 @@ FSubobjectDataHandle FSubobjectData::FindChild(const FSubobjectDataHandle& InChi
 
 FSubobjectDataHandle FSubobjectData::FindChildByObject(UObject* ContextObject) const
 {
+	REDIRECT_TO_ADAPTER(FindChildByObject, ContextObject)
+
 	if (!ContextObject)
 	{
 		return FSubobjectDataHandle::InvalidHandle;
@@ -685,6 +757,8 @@ bool FSubobjectData::RemoveChildHandleOnly(const FSubobjectDataHandle& InHandle)
 
 FText FSubobjectData::GetAssetName() const
 {
+	REDIRECT_TO_ADAPTER(GetAssetName)
+	
 	UActorComponent* Template = GetMutableComponentTemplate();
 
 	FText AssetName = LOCTEXT("None", "None");
@@ -702,6 +776,8 @@ FText FSubobjectData::GetAssetName() const
 
 FText FSubobjectData::GetAssetPath() const
 {
+	REDIRECT_TO_ADAPTER(GetAssetPath)
+
 	FText AssetName = LOCTEXT("None", "None");
 
 	UActorComponent* Template = GetMutableComponentTemplate();
@@ -720,6 +796,8 @@ FText FSubobjectData::GetAssetPath() const
 
 bool FSubobjectData::IsAssetVisible() const
 {
+	REDIRECT_TO_ADAPTER(IsAssetVisible)
+
 	UActorComponent* Template = GetMutableComponentTemplate();
 
 	if (Template && FComponentAssetBrokerage::SupportsAssets(Template))
@@ -732,6 +810,8 @@ bool FSubobjectData::IsAssetVisible() const
 
 FText FSubobjectData::GetMobilityToolTipText() const
 {
+	REDIRECT_TO_ADAPTER(GetMobilityToolTipText)
+
 	FText MobilityToolTip;
     
     if (const USceneComponent* SceneComponentTemplate = Cast<USceneComponent>(GetComponentTemplate()))
@@ -765,6 +845,8 @@ FText FSubobjectData::GetMobilityToolTipText() const
 
 FText FSubobjectData::GetComponentEditorOnlyTooltipText() const
 {
+	REDIRECT_TO_ADAPTER(GetComponentEditorOnlyTooltipText)
+
 	FText ComponentType = LOCTEXT("ComponentEditorOnlyFalse", "False");
 	
 	if (IsComponent())
@@ -786,6 +868,8 @@ FText FSubobjectData::GetComponentEditorOnlyTooltipText() const
 
 FText FSubobjectData::GetIntroducedInToolTipText() const
 {
+	REDIRECT_TO_ADAPTER(GetIntroducedInToolTipText)
+
 	FText IntroducedInTooltip = LOCTEXT("IntroducedInThisBPTooltip", "this class");
 
 	if (IsInheritedComponent())
@@ -873,6 +957,8 @@ FText FSubobjectData::GetIntroducedInToolTipText() const
 
 FText FSubobjectData::GetActorDisplayText() const
 {
+	REDIRECT_TO_ADAPTER(GetActorDisplayText)
+
 	if (bIsChildActor)
 	{
 		if (const AActor* ChildActor = GetObject<AActor>())
@@ -898,8 +984,20 @@ FText FSubobjectData::GetActorDisplayText() const
 	return FText::GetEmpty();
 }
 
+const UObject* FSubobjectData::GetObject(bool bEvenIfPendingKill) const
+{
+	return WeakObjectPtr.Get(bEvenIfPendingKill);
+}
+
+const TInstancedStruct<FSubobjectDataSubsystemContextDataBase>& FSubobjectData::GetContextData() const
+{
+	return ContextData;
+}
+
 bool FSubobjectData::IsInstancedActor() const
 {
+	REDIRECT_TO_ADAPTER(IsInstancedActor)
+
 	if (const AActor* Actor = GetObject<AActor>())
 	{
 		return !Actor->IsTemplate();
@@ -910,6 +1008,8 @@ bool FSubobjectData::IsInstancedActor() const
 
 bool FSubobjectData::IsNativeComponent() const
 {
+	REDIRECT_TO_ADAPTER(IsNativeComponent)
+
 	if (const UActorComponent* Template = GetComponentTemplate())
 	{
 		return Template->CreationMethod == EComponentCreationMethod::Native && GetSCSNode() == nullptr;
@@ -920,6 +1020,8 @@ bool FSubobjectData::IsNativeComponent() const
 
 bool FSubobjectData::IsBlueprintInheritedComponent() const
 {
+	REDIRECT_TO_ADAPTER(IsBlueprintInheritedComponent)
+
 	if (const UActorComponent* Template = GetComponentTemplate())
 	{
 		if(GetSCSNode() != nullptr)
@@ -933,6 +1035,8 @@ bool FSubobjectData::IsBlueprintInheritedComponent() const
 
 bool FSubobjectData::IsInheritedComponent() const
 {
+	REDIRECT_TO_ADAPTER(IsInheritedComponent)
+
 	// This covers a component that is added via blueprints
 	if(GetSCSNode() != nullptr)
 	{
@@ -948,11 +1052,15 @@ bool FSubobjectData::IsInheritedComponent() const
 
 bool FSubobjectData::IsSceneComponent() const
 {
+	REDIRECT_TO_ADAPTER(IsSceneComponent)
+
 	return Cast<USceneComponent>(GetComponentTemplate()) != nullptr;
 }
 
 bool FSubobjectData::IsRootComponent() const
 {
+	REDIRECT_TO_ADAPTER(IsRootComponent)
+
 	const UActorComponent* ComponentTemplate = GetComponentTemplate();
 	
 	if(!ComponentTemplate)
@@ -987,6 +1095,8 @@ bool FSubobjectData::IsRootComponent() const
 
 bool FSubobjectData::IsDefaultSceneRoot() const
 {
+	REDIRECT_TO_ADAPTER(IsDefaultSceneRoot)
+	
 	// If this is a scene component and is instanced, then it will have a specific name
 	const USceneComponent* SceneComponent = Cast<USceneComponent>(GetComponentTemplate());
 	if (SceneComponent && !SceneComponent->IsTemplate())
@@ -1020,6 +1130,8 @@ bool FSubobjectData::IsDefaultSceneRoot() const
 
 bool FSubobjectData::SceneRootHasDefaultName() const
 {
+	REDIRECT_TO_ADAPTER(SceneRootHasDefaultName)
+
 	const UActorComponent* Template = GetComponentTemplate();
 	const FName TemplateName = Template ? Template->GetFName() : NAME_None;
 
@@ -1034,17 +1146,23 @@ bool FSubobjectData::SceneRootHasDefaultName() const
 
 bool FSubobjectData::IsComponent() const
 {
+	REDIRECT_TO_ADAPTER(IsComponent)
+
 	// Check if we are pointing to a component
 	return GetComponentTemplate() != nullptr;
 }
 
 bool FSubobjectData::IsChildActor() const
 {
+	REDIRECT_TO_ADAPTER(IsChildActor)
+
 	return bIsChildActor;
 }
 
 bool FSubobjectData::IsChildActorSubtreeObject() const
 {
+	REDIRECT_TO_ADAPTER(IsChildActorSubtreeObject)
+
 	const FSubobjectDataHandle& RootActor = GetRootSubobject();
 	if (const FSubobjectData* RootData = RootActor.GetData())
 	{
@@ -1055,6 +1173,8 @@ bool FSubobjectData::IsChildActorSubtreeObject() const
 
 bool FSubobjectData::IsRootActor() const
 {
+	REDIRECT_TO_ADAPTER(IsRootActor)
+	
 	// This is the root actor if it points to an AActor and has no parent
 	if(const AActor* Actor = GetObject<AActor>())
 	{
@@ -1065,11 +1185,15 @@ bool FSubobjectData::IsRootActor() const
 
 bool FSubobjectData::IsActor() const
 {
+	REDIRECT_TO_ADAPTER(IsActor)
+	
 	return GetObject<AActor>() != nullptr;
 }
 
 bool FSubobjectData::IsInstancedInheritedComponent() const
 {
+	REDIRECT_TO_ADAPTER(IsInstancedInheritedComponent)
+
 	if(!IsComponent())
 	{
 		return false;
@@ -1088,6 +1212,8 @@ bool FSubobjectData::IsInstancedInheritedComponent() const
 
 bool FSubobjectData::IsAttachedTo(const FSubobjectDataHandle& InHandle) const
 {
+	REDIRECT_TO_ADAPTER(IsAttachedTo, InHandle)
+
 	FSubobjectDataHandle TestParentHandle = ParentObjectHandle;
 
 	while (TestParentHandle.IsValid())
@@ -1181,9 +1307,4 @@ bool FSubobjectData::IsInheritedSCSNode() const
 	return bIsInheritedSCS;
 }
 
-const UChildActorComponent* FSubobjectData::GetChildActorComponent(bool bEvenIfPendingKill) const
-{
-	return GetObject<UChildActorComponent>(bEvenIfPendingKill);
-}
-
 #undef LOCTEXT_NAMESPACE
diff --git a/Engine/Source/Editor/SubobjectDataInterface/Private/SubobjectDataAdapterBase.cpp b/Engine/Source/Editor/SubobjectDataInterface/Private/SubobjectDataAdapterBase.cpp
new file mode 100644
--- /dev/null
+++ b/Engine/Source/Editor/SubobjectDataInterface/Private/SubobjectDataAdapterBase.cpp
@@ -0,0 +1,274 @@
+﻿// Copyright Epic Games, Inc. All Rights Reserved.
+
+#include "SubobjectDataAdapterBase.h"
+
+#include "SubobjectDataHandle.h"
+
+#define LOCTEXT_NAMESPACE "FSubobjectDataBehaviorBase"
+
+FSubobjectDataSubsystemContextData_SingleUObjectContextObject::FSubobjectDataSubsystemContextData_SingleUObjectContextObject()
+{
+}
+
+FSubobjectDataSubsystemContextData_SingleUObjectContextObject::FSubobjectDataSubsystemContextData_SingleUObjectContextObject(UObject* ContextObject)
+	: Object(ContextObject)
+{
+}
+
+FSubobjectDataSubsystemContextData_TedsRow::FSubobjectDataSubsystemContextData_TedsRow()
+{
+}
+
+FSubobjectDataSubsystemContextData_TedsRow::FSubobjectDataSubsystemContextData_TedsRow(UE::Editor::DataStorage::RowHandle InRowHandle)
+{
+	RowHandle = InRowHandle;
+}
+
+FSubobjectDataSubsystemContextData_TedsRow::FSubobjectDataSubsystemContextData_TedsRow(const FTedsRowHandle& InRowHandle)
+	: RowHandle(InRowHandle)
+{
+}
+
+FSubobjectDataSubsystemAdapterBase::~FSubobjectDataSubsystemAdapterBase() = default;
+
+void FSubobjectDataSubsystemAdapterBase::GatherSubobjectData(const TInstancedStruct<FSubobjectDataSubsystemContextDataBase>& ContextData, TArray<FSubobjectDataHandle>& OutArray, TFunctionRef<FSubobjectDataHandle()> CreateSubobjectData) const
+{
+}
+
+FSubobjectDataAdapterBase::~FSubobjectDataAdapterBase() = default;
+
+bool FSubobjectDataAdapterBase::CanEdit(const FSubobjectData& Data) const
+{
+	return false;
+}
+
+bool FSubobjectDataAdapterBase::CanDelete(const FSubobjectData& Data) const
+{
+	return false;
+}
+
+bool FSubobjectDataAdapterBase::CanDuplicate(const FSubobjectData& Data) const
+{
+	return false;
+}
+
+bool FSubobjectDataAdapterBase::CanCopy(const FSubobjectData& Data) const
+{
+	return false;
+}
+
+bool FSubobjectDataAdapterBase::CanReparent(const FSubobjectData& Data) const
+{
+	return false;
+}
+
+bool FSubobjectDataAdapterBase::CanRename(const FSubobjectData& Data) const
+{
+	return false;
+}
+
+const UObject* FSubobjectDataAdapterBase::GetObject(const FSubobjectData& Data, bool bEventIfPendingKill) const
+{
+	return nullptr;
+}
+
+const UObject* FSubobjectDataAdapterBase::GetObjectForBlueprint(const FSubobjectData& Data, UBlueprint* Blueprint) const
+{
+	return nullptr;
+}
+
+const UActorComponent* FSubobjectDataAdapterBase::GetComponentTemplate(const FSubobjectData& Data, bool bEvenIfPendingKill) const
+{
+	return nullptr;
+}
+
+const UActorComponent* FSubobjectDataAdapterBase::FindComponentInstanceInActor(const FSubobjectData& Data, const AActor* InActor) const
+{
+	return nullptr;
+}
+
+UActorComponent* FSubobjectDataAdapterBase::FindMutableComponentInstanceInActor(const FSubobjectData& Data, const AActor* InActor) const
+{
+	return nullptr;
+}
+
+UBlueprint* FSubobjectDataAdapterBase::GetBlueprint(const FSubobjectData& Data) const
+{
+	return nullptr;
+}
+
+UBlueprint* FSubobjectDataAdapterBase::GetBlueprintBeingEdited(const FSubobjectData& Data) const
+{
+	return nullptr;
+}
+
+bool FSubobjectDataAdapterBase::IsInstancedComponent(const FSubobjectData& Data) const
+{
+	return false;
+}
+
+bool FSubobjectDataAdapterBase::IsInstancedActor(const FSubobjectData& Data) const
+{
+	return false;
+}
+
+bool FSubobjectDataAdapterBase::IsNativeComponent(const FSubobjectData& Data) const
+{
+	return false;
+}
+
+bool FSubobjectDataAdapterBase::IsBlueprintInheritedComponent(const FSubobjectData& Data) const
+{
+	return false;
+}
+
+bool FSubobjectDataAdapterBase::IsInheritedComponent(const FSubobjectData& Data) const
+{
+	return false;
+}
+
+bool FSubobjectDataAdapterBase::IsSceneComponent(const FSubobjectData& Data) const
+{
+	return false;
+}
+
+bool FSubobjectDataAdapterBase::IsRootComponent(const FSubobjectData& Data) const
+{
+	return false;
+}
+
+bool FSubobjectDataAdapterBase::IsDefaultSceneRoot(const FSubobjectData& Data) const
+{
+	return false;
+}
+
+bool FSubobjectDataAdapterBase::SceneRootHasDefaultName(const FSubobjectData& Data) const
+{
+	return false;
+}
+
+bool FSubobjectDataAdapterBase::IsComponent(const FSubobjectData& Data) const
+{
+	return false;
+}
+
+bool FSubobjectDataAdapterBase::IsChildActor(const FSubobjectData& Data) const
+{
+	return false;
+}
+
+bool FSubobjectDataAdapterBase::IsChildActorSubtreeObject(const FSubobjectData& Data) const
+{
+	return false;
+}
+
+bool FSubobjectDataAdapterBase::IsRootActor(const FSubobjectData& Data) const
+{
+	return false;
+}
+
+bool FSubobjectDataAdapterBase::IsActor(const FSubobjectData& Data) const
+{
+	return false;
+}
+
+bool FSubobjectDataAdapterBase::IsInstancedInheritedComponent(const FSubobjectData& Data) const
+{
+	return false;
+}
+
+bool FSubobjectDataAdapterBase::IsAttachedTo(const FSubobjectData& Data, const FSubobjectDataHandle& InHandle) const
+{
+	return false;
+}
+
+FString FSubobjectDataAdapterBase::GetDisplayString(const FSubobjectData& Data, bool bShowNativeComponentNames) const
+{
+	return TEXT("");
+}
+
+FText FSubobjectDataAdapterBase::GetDragDropDisplayText(const FSubobjectData& Data) const
+{
+	return FText();
+}
+
+FText FSubobjectDataAdapterBase::GetDisplayNameContextModifiers(const FSubobjectData& Data,
+	bool bShowNativeComponentNames) const
+{
+	return FText();
+}
+
+FText FSubobjectDataAdapterBase::GetDisplayName(const FSubobjectData& Data) const
+{
+	return FText();
+}
+
+FName FSubobjectDataAdapterBase::GetVariableName(const FSubobjectData& Data) const
+{
+	return FName();
+}
+
+FText FSubobjectDataAdapterBase::GetSocketName(const FSubobjectData& Data) const
+{
+	return FText();
+}
+
+FName FSubobjectDataAdapterBase::GetSocketFName(const FSubobjectData& Data) const
+{
+	return FName();
+}
+
+bool FSubobjectDataAdapterBase::HasValidSocket(const FSubobjectData& Data) const
+{
+	return false;
+}
+
+void FSubobjectDataAdapterBase::SetSocketName(FSubobjectData& Data, FName InNewName) const
+{
+}
+
+void FSubobjectDataAdapterBase::SetupAttachment(FSubobjectData& Data, FName SocketName, const FSubobjectDataHandle& AttachParentHandle) const
+{
+}
+
+FSubobjectDataHandle FSubobjectDataAdapterBase::FindChildByObject(const FSubobjectData& Data, UObject* ContextObject) const
+{
+	return FSubobjectDataHandle();
+}
+
+FText FSubobjectDataAdapterBase::GetAssetName(const FSubobjectData& Data) const
+{
+	return FText();
+}
+
+FText FSubobjectDataAdapterBase::GetAssetPath(const FSubobjectData& Data) const
+{
+	return FText();
+}
+
+bool FSubobjectDataAdapterBase::IsAssetVisible(const FSubobjectData& Data) const
+{
+	return false;
+}
+
+FText FSubobjectDataAdapterBase::GetMobilityToolTipText(const FSubobjectData& Data) const
+{
+	return FText();
+}
+
+FText FSubobjectDataAdapterBase::GetComponentEditorOnlyTooltipText(const FSubobjectData& Data) const
+{
+	return FText();
+}
+
+FText FSubobjectDataAdapterBase::GetIntroducedInToolTipText(const FSubobjectData& Data) const
+{
+	return FText();
+}
+
+FText FSubobjectDataAdapterBase::GetActorDisplayText(const FSubobjectData& Data) const
+{
+	return FText();
+}
+
+#undef LOCTEXT_NAMESPACE
\ No newline at end of file
diff --git a/Engine/Source/Editor/SubobjectDataInterface/Private/SubobjectDataSubsystem.cpp b/Engine/Source/Editor/SubobjectDataInterface/Private/SubobjectDataSubsystem.cpp
--- a/Engine/Source/Editor/SubobjectDataInterface/Private/SubobjectDataSubsystem.cpp
+++ b/Engine/Source/Editor/SubobjectDataInterface/Private/SubobjectDataSubsystem.cpp
@@ -31,9 +31,15 @@
 
 #include "Kismet2/CompilerResultsLog.h"		// Adding compiler errors to nodes that have their variables deleted
 #include "K2Node_ComponentBoundEvent.h"
+#include "DataStorage/Features.h"
+#include "Elements/Columns/TypedElementHandleColumn.h"
+#include "Elements/Framework/TypedElementRegistry.h"
+#include "Elements/Interfaces/TypedElementDataStorageCompatibilityInterface.h"
+#include "Elements/Interfaces/TypedElementHierarchyInterface.h"
 
 #include "Engine/SCS_Node.h"		// #TODO_BH  We need to remove this when the actual subobject refactor happens
 #include "Logging/MessageLog.h"
+#include "StructUtils/StructView.h"
 
 #include UE_INLINE_GENERATED_CPP_BY_NAME(SubobjectDataSubsystem)
 
@@ -77,6 +83,82 @@ USubobjectDataSubsystem* USubobjectDataSubsystem::Get()
 	return GEngine->GetEngineSubsystem<USubobjectDataSubsystem>();
 }
 
+FSubobjectDataAdapterHandle USubobjectDataSubsystem::RegisterSubobjectDataAdapter(
+	FSubobjectDataAdapterFilter InSubobjectDataAdapterFilter,
+	TUniquePtr<FSubobjectDataAdapterBase> InDataAdapter)
+{
+	FSubobjectDataAdapterHandle Handle;
+	// Pre-increment so that invalid handles are 0
+	Handle.Opaque = ++DataAdapterHandleIndex;
+
+	DataAdapters.Emplace(FSubobjectDataAdapterEntry
+		{
+			.Adapter = MoveTemp(InDataAdapter),
+			.Handle = Handle,
+			.Filter = MoveTemp(InSubobjectDataAdapterFilter)
+		});
+	return Handle;
+}
+
+void USubobjectDataSubsystem::UnregisterSubobjectDataAdapter(FSubobjectDataAdapterHandle Handle)
+{
+	DataAdapters.RemoveAll([Handle](const FSubobjectDataAdapterEntry& Entry)
+		{
+			return Entry.Handle.Opaque == Handle.Opaque;
+		});
+}
+
+const FSubobjectDataAdapterBase* USubobjectDataSubsystem::FindDataAdapter(const FSubobjectData& SubobjectData) const
+{
+	for (const FSubobjectDataAdapterEntry& Entry : DataAdapters)
+	{
+		if (Entry.Filter(SubobjectData))
+		{
+			return Entry.Adapter.Get();
+		}
+	}
+	return nullptr;
+}
+
+
+FSubobjectDataSubsystemAdapterHandle USubobjectDataSubsystem::RegisterSubsystemAdapter(
+	FSubobjectSubsystemAdapterFilter InFilter,
+	TUniquePtr<FSubobjectDataSubsystemAdapterBase> InAdapter)
+{
+	FSubobjectDataSubsystemAdapterHandle Handle;
+	// Pre-increment so that invalid handles are 0
+	Handle.Opaque = ++SubsystemAdapterHandleIndex;
+	
+	SubsystemAdapters.Emplace(FSubsystemAdapterEntry
+		{
+			.Adapter = MoveTemp(InAdapter),
+			.Handle = Handle,
+			.Filter = MoveTemp(InFilter)
+		});
+	return Handle;
+}
+
+void USubobjectDataSubsystem::UnregisterSubsystemAdapter(FSubobjectDataSubsystemAdapterHandle Handle)
+{
+	SubsystemAdapters.RemoveAll([Handle](const FSubsystemAdapterEntry& Entry)
+	{
+		return Entry.Handle.Opaque == Handle.Opaque;
+	});
+}
+
+const FSubobjectDataSubsystemAdapterBase* USubobjectDataSubsystem::FindSubsystemAdapter(
+	const TInstancedStruct<FSubobjectDataSubsystemContextDataBase>& Struct) const
+{
+	for (const FSubsystemAdapterEntry& Entry : SubsystemAdapters)
+	{
+		if (Entry.Filter(Struct))
+		{
+			return Entry.Adapter.Get();
+		}
+	}
+	return nullptr;
+}
+
 void USubobjectDataSubsystem::K2_GatherSubobjectDataForBlueprint(UBlueprint* Context, TArray<FSubobjectDataHandle>& OutArray)
 {
 	// Return the current CDO that was last generated for the class
@@ -206,6 +288,20 @@ FSubobjectDataHandle USubobjectDataSubsystem::FindAttachParentForInheritedCompon
 
 void USubobjectDataSubsystem::GatherSubobjectData(UObject* Context, TArray<FSubobjectDataHandle>& OutArray)
 {
+	OutArray.Reset();
+	
+	FSubobjectDataSubsystemContextData_SingleUObjectContextObject ContextDataAdapter_OnStack(Context);
+	TConstStructView<FSubobjectDataSubsystemContextDataBase> StructView(ContextDataAdapter_OnStack);
+	TInstancedStruct<FSubobjectDataSubsystemContextDataBase> ContextDataAdapter(StructView);
+	
+	if (const FSubobjectDataSubsystemAdapterBase* Adapter = FindSubsystemAdapter(ContextDataAdapter))
+	{
+		return Adapter->GatherSubobjectData(
+			ContextDataAdapter,
+			OutArray,
+			[this]() { return this->CreateSubobjectData(); });
+	}
+	
 	if (!Context)
 	{
 		UE_LOG(LogSubobjectSubsystem, Warning, TEXT("Could not gather subobject data, there was no context given!"));
@@ -215,11 +311,8 @@ void USubobjectDataSubsystem::GatherSubobjectData(UObject* Context, TArray<FSubo
 	const AActor* ActorContext = Cast<AActor>(Context);
 	if(!ActorContext)
 	{
-		UE_LOG(LogSubobjectSubsystem, Warning, TEXT("Could not gather subobject data, the given context was not an actor!"));
 		return;
 	}
-	
-	OutArray.Reset();
 
 	FSubobjectDataHandle RootActorHandle = CreateSubobjectData(Context);
 	OutArray.Add(RootActorHandle);
@@ -2615,14 +2708,25 @@ FScopedTransaction* USubobjectDataSubsystem::BeginTransaction(const TArray<FSubo
 	return OutTransaction;
 }
 
+FSubobjectDataHandle USubobjectDataSubsystem::CreateSubobjectData()
+{
+	TSharedPtr<FSubobjectData> SharedPtr = MakeShared<FSubobjectData>();
+	SharedPtr->Handle.DataPtr = SharedPtr;
+	return SharedPtr->GetHandle();
+}
+
 FSubobjectDataHandle USubobjectDataSubsystem::CreateSubobjectData(UObject* Context, const FSubobjectDataHandle& ParentHandle /* = FSubobjectDataHandle::InvalidHandle */, bool bIsInheritedSCS/* = false */)
 {
+	PRAGMA_DISABLE_DEPRECATION_WARNINGS
 	TSharedPtr<FSubobjectData> SharedPtr = TSharedPtr<FSubobjectData>(new FSubobjectData(Context, ParentHandle, bIsInheritedSCS));
+	PRAGMA_ENABLE_DEPRECATION_WARNINGS
 	
 	if(!SharedPtr.IsValid())
 	{
 		ensureMsgf(false, TEXT("The subobject data factories failed to create subobject data!"));
+		PRAGMA_DISABLE_DEPRECATION_WARNINGS
 		SharedPtr = TSharedPtr<FSubobjectData>(new FSubobjectData(Context, ParentHandle, bIsInheritedSCS));
+		PRAGMA_ENABLE_DEPRECATION_WARNINGS
 	}
 
 	check(SharedPtr.IsValid());
diff --git a/Engine/Source/Editor/SubobjectDataInterface/Public/SubobjectData.h b/Engine/Source/Editor/SubobjectDataInterface/Public/SubobjectData.h
--- a/Engine/Source/Editor/SubobjectDataInterface/Public/SubobjectData.h
+++ b/Engine/Source/Editor/SubobjectDataInterface/Public/SubobjectData.h
@@ -8,12 +8,14 @@
 #include "Components/ActorComponent.h"
 #include "Kismet/BlueprintFunctionLibrary.h"
 #include "Engine/SCS_Node.h"		// #TODO_BH  We need to remove this when the actual subobject refactor happens
+#include "StructUtils/InstancedStruct.h"
 
 #include "SubobjectData.generated.h"
 
 #define UE_API SUBOBJECTDATAINTERFACE_API
 
 class UBlueprint;
+struct FSubobjectDataSubsystemContextDataBase;
 
 /**
 * A struct that represents a single subobject. This can be anything, but are
@@ -25,15 +27,32 @@ class UBlueprint;
 USTRUCT(BlueprintType)
 struct FSubobjectData final
 {
-friend class USubobjectDataSubsystem;
+	friend class USubobjectDataSubsystem;
 public:
 	GENERATED_BODY()
 
 	UE_API FSubobjectData();
+
+	UE_DEPRECATED(5.7, "")
 	UE_API FSubobjectData(UObject* ContextObject, const FSubobjectDataHandle& ParentHandle, const bool bIsInheritedSCS);
 
+	UE_API void Initialize(TInstancedStruct<FSubobjectDataSubsystemContextDataBase> ContextData, const FSubobjectDataHandle& ParentHandle);
+	
 	/** True if this subobject's handle is valid */
 	bool IsValid() const { return Handle.IsValid(); }
+
+	/**
+	* @return	A read-only reference to the object represented by this node.
+	*/
+	UE_API const UObject* GetObject(bool bEvenIfPendingKill = false) const;
+
+	/**
+	* @return A typed struct containing the context object for this SubobjectData
+	*         Discrimination types are:
+	*           - FSubobjectDataSubsystemContextData_SingleUObjectContextObject
+	*           - FSubobjectDataSubsystemContextData_TedsRow
+	*/
+	UE_API const TInstancedStruct<FSubobjectDataSubsystemContextDataBase>& GetContextData() const;
 	
 	/**
 	 * @return Whether or not we can edit properties for this subobject
@@ -67,11 +86,6 @@ public:
 	*/
 	UE_API bool CanRename() const;
 
-	/**
-	* @return	A read-only reference to the object represented by this node.
-	*/
-	inline const UObject* GetObject(bool bEvenIfPendingKill = false) const { return WeakObjectPtr.Get(bEvenIfPendingKill); }
-
 	/**
      * @param	InBlueprint			The Blueprint in which the object will be edited.
      *
@@ -212,7 +226,7 @@ public:
 	UE_API FText GetIntroducedInToolTipText() const;
 
 	UE_API FText GetActorDisplayText() const;
-		
+
 private:
 	/**
 	* Add the given subobject handle as a child of this.
@@ -249,6 +263,13 @@ private:
 	UPROPERTY()
 	TArray<FSubobjectDataHandle> ChildrenHandles;
 
+	// ContextData is a replacement for ContextObject as a mechanism to allow SubobjectData
+	// to serve as an interface for multiple data types (not just single UObject) but multiple non-UObject
+	// ContextObject will still be populated with a single UObject if available, but ContextData
+	// should be the preferred data source for the SubobjectData instances.
+	UPROPERTY()
+	TInstancedStruct<FSubobjectDataSubsystemContextDataBase> ContextData;
+
 	////////////////////////////////////////////
 	// Muteable accessors for this subobject to be used by the subsystem
 	inline UObject* GetMutableObject(bool bEvenIfPendingKill = false) const { return WeakObjectPtr.Get(bEvenIfPendingKill); }
@@ -305,7 +326,6 @@ private:
 	 * which means that we must use the InheritableComponentHandler
 	 */
 	UE_API bool IsInheritedSCSNode() const;
-	UE_API const UChildActorComponent* GetChildActorComponent(bool bEvenIfPendingKill = false) const;
 
 	/** True if this SCS node is inherited from another blueprint generated class. */
 	bool bIsInheritedSubobject = false;
diff --git a/Engine/Source/Editor/SubobjectDataInterface/Public/SubobjectDataAdapterBase.h b/Engine/Source/Editor/SubobjectDataInterface/Public/SubobjectDataAdapterBase.h
new file mode 100644
--- /dev/null
+++ b/Engine/Source/Editor/SubobjectDataInterface/Public/SubobjectDataAdapterBase.h
@@ -0,0 +1,145 @@
+﻿// Copyright Epic Games, Inc. All Rights Reserved.
+
+#pragma once
+
+#include "CoreMinimal.h"
+// #include "SubobjectDataHandle.h"
+#include "DataStorage/Handles.h"
+#include "StructUtils/InstancedStruct.h"
+
+#include "SubobjectDataAdapterBase.generated.h"
+
+class AActor;
+class UActorComponent;
+class UBlueprint;
+class UObject;
+struct FSubobjectData;
+struct FSubobjectDataHandle;
+
+struct FTedsRowHandle;
+
+#define UE_API SUBOBJECTDATAINTERFACE_API
+
+struct
+FSubobjectDataAdapterHandle
+{
+	friend class USubobjectDataSubsystem;
+private:
+	uint64 Opaque = 0;
+};
+
+struct
+FSubobjectDataSubsystemAdapterHandle
+{
+	friend class USubobjectDataSubsystem;
+private:
+	uint64 Opaque = 0;
+};
+
+USTRUCT()
+struct
+FSubobjectDataSubsystemContextDataBase
+{
+	GENERATED_BODY()
+};
+
+// An implementation of the ContextData that supports containing a single UObject
+// This is to provide interoperability with the existing ContextObject used in APIs of the SubobjectDataSubsystem
+USTRUCT()
+struct
+FSubobjectDataSubsystemContextData_SingleUObjectContextObject : public FSubobjectDataSubsystemContextDataBase
+{
+	GENERATED_BODY()
+
+	UE_API FSubobjectDataSubsystemContextData_SingleUObjectContextObject();
+	UE_API explicit FSubobjectDataSubsystemContextData_SingleUObjectContextObject(UObject* ContextObject);
+
+	UPROPERTY()
+	TWeakObjectPtr<UObject> Object;
+};
+
+USTRUCT()
+struct
+FSubobjectDataSubsystemContextData_TedsRow : public FSubobjectDataSubsystemContextDataBase
+{
+	GENERATED_BODY()
+
+	UE_API FSubobjectDataSubsystemContextData_TedsRow();
+	UE_API explicit FSubobjectDataSubsystemContextData_TedsRow(UE::Editor::DataStorage::RowHandle InRowHandle);
+	UE_API explicit FSubobjectDataSubsystemContextData_TedsRow(const FTedsRowHandle& InRowHandle);
+
+	UPROPERTY()
+	FTedsRowHandle RowHandle;
+};
+
+class FSubobjectDataSubsystemAdapterBase
+{
+public:
+	UE_API virtual ~FSubobjectDataSubsystemAdapterBase();
+	UE_API virtual void GatherSubobjectData(const TInstancedStruct<FSubobjectDataSubsystemContextDataBase>& ContextData, TArray<FSubobjectDataHandle>& OutArray, TFunctionRef<FSubobjectDataHandle()> CreateSubobjectData) const;
+};
+
+
+class FSubobjectDataAdapterBase
+{
+public:
+	UE_API virtual ~FSubobjectDataAdapterBase();
+
+	UE_API virtual bool CanEdit(const FSubobjectData& Data) const;
+	UE_API virtual bool CanDelete(const FSubobjectData& Data) const;
+	UE_API virtual bool CanDuplicate(const FSubobjectData& Data) const;
+	UE_API virtual bool CanCopy(const FSubobjectData& Data) const;
+	UE_API virtual bool CanReparent(const FSubobjectData& Data) const;
+	UE_API virtual bool CanRename(const FSubobjectData& Data) const;
+
+	UE_API virtual const UObject* GetObject(const FSubobjectData& Data, bool bEventIfPendingKill = false) const;
+	UE_API virtual const UObject* GetObjectForBlueprint(const FSubobjectData& Data, UBlueprint* Blueprint) const;
+	UE_API virtual const UActorComponent* GetComponentTemplate(const FSubobjectData& Data, bool bEvenIfPendingKill = false) const;
+	UE_API virtual const UActorComponent* FindComponentInstanceInActor(const FSubobjectData& Data, const AActor* InActor) const;
+	UE_API virtual UActorComponent* FindMutableComponentInstanceInActor(const FSubobjectData& Data, const AActor* InActor) const;
+
+	UE_API virtual UBlueprint* GetBlueprint(const FSubobjectData& Data) const;
+	UE_API virtual UBlueprint* GetBlueprintBeingEdited(const FSubobjectData& Data) const;
+
+	UE_API virtual bool IsInstancedComponent(const FSubobjectData& Data) const;
+	UE_API virtual bool IsInstancedActor(const FSubobjectData& Data) const;
+	UE_API virtual bool IsNativeComponent(const FSubobjectData& Data) const;
+	UE_API virtual bool IsBlueprintInheritedComponent(const FSubobjectData& Data) const;
+	UE_API virtual bool IsInheritedComponent(const FSubobjectData& Data) const;
+	UE_API virtual bool IsSceneComponent(const FSubobjectData& Data) const;
+	UE_API virtual bool IsRootComponent(const FSubobjectData& Data) const;
+	UE_API virtual bool IsDefaultSceneRoot(const FSubobjectData& Data) const;
+	UE_API virtual bool SceneRootHasDefaultName(const FSubobjectData& Data) const;
+	UE_API virtual bool IsComponent(const FSubobjectData& Data) const;
+	UE_API virtual bool IsChildActor(const FSubobjectData& Data) const;
+	UE_API virtual bool IsChildActorSubtreeObject(const FSubobjectData& Data) const;
+	UE_API virtual bool IsRootActor(const FSubobjectData& Data) const;
+	UE_API virtual bool IsActor(const FSubobjectData& Data) const;
+	UE_API virtual bool IsInstancedInheritedComponent(const FSubobjectData& Data) const;
+	UE_API virtual bool IsAttachedTo(const FSubobjectData& Data, const FSubobjectDataHandle& InHandle) const;
+	UE_API virtual FString GetDisplayString(const FSubobjectData& Data, bool bShowNativeComponentNames = true) const;
+	UE_API virtual FText GetDragDropDisplayText(const FSubobjectData& Data) const;
+
+	UE_API virtual FText GetDisplayNameContextModifiers(const FSubobjectData& Data, bool bShowNativeComponentNames = true) const;
+	
+	UE_API virtual FText GetDisplayName(const FSubobjectData& Data) const;
+
+	UE_API virtual FName GetVariableName(const FSubobjectData& Data) const;
+
+	// Sockets for attaching in the viewport
+	UE_API virtual FText GetSocketName(const FSubobjectData& Data) const;
+	UE_API virtual FName GetSocketFName(const FSubobjectData& Data) const;
+	UE_API virtual bool HasValidSocket(const FSubobjectData& Data) const;
+	UE_API virtual void SetSocketName(FSubobjectData& Data, FName InNewName) const;
+	UE_API virtual void SetupAttachment(FSubobjectData& Data, FName SocketName, const FSubobjectDataHandle& AttachParentHandle) const;
+	
+	UE_API virtual FSubobjectDataHandle FindChildByObject(const FSubobjectData& Data, UObject* ContextObject) const;
+	UE_API virtual FText GetAssetName(const FSubobjectData& Data) const;
+	UE_API virtual FText GetAssetPath(const FSubobjectData& Data) const;
+	UE_API virtual bool IsAssetVisible(const FSubobjectData& Data) const;
+	UE_API virtual FText GetMobilityToolTipText(const FSubobjectData& Data) const;
+	UE_API virtual FText GetComponentEditorOnlyTooltipText(const FSubobjectData& Data) const;
+	UE_API virtual FText GetIntroducedInToolTipText(const FSubobjectData& Data) const;
+
+	UE_API virtual FText GetActorDisplayText(const FSubobjectData& Data) const;
+};
\ No newline at end of file
diff --git a/Engine/Source/Editor/SubobjectDataInterface/Public/SubobjectDataSubsystem.h b/Engine/Source/Editor/SubobjectDataInterface/Public/SubobjectDataSubsystem.h
--- a/Engine/Source/Editor/SubobjectDataInterface/Public/SubobjectDataSubsystem.h
+++ b/Engine/Source/Editor/SubobjectDataInterface/Public/SubobjectDataSubsystem.h
@@ -7,6 +7,7 @@
 #include "Subsystems/EngineSubsystem.h"
 #include "SubobjectData.h"
 #include "ModuleDescriptor.h"
+#include "SubobjectDataAdapterBase.h"
 #include "SubobjectDataSubsystem.generated.h"
 
 #define UE_API SUBOBJECTDATAINTERFACE_API
@@ -124,6 +125,25 @@ public:
 	/** Delegate invoked when a new subobject is successfully added */
 	FOnNewSubobjectAdded& OnNewSubobjectAdded() { return OnNewSubobjectAdded_Delegate; };
 
+	using FSubobjectDataAdapterFilter = TFunction<bool(const FSubobjectData&)>;
+	/** */
+	UE_EXPERIMENTAL(5.7, "Experimental adapter to adjust behavior of SubobjectData")
+	UE_API FSubobjectDataAdapterHandle RegisterSubobjectDataAdapter(FSubobjectDataAdapterFilter InFilter, TUniquePtr<FSubobjectDataAdapterBase> InDataAdapter);
+	UE_EXPERIMENTAL(5.7, "Experimental adapter to adjust behavior of SubobjectData")
+	UE_API void UnregisterSubobjectDataAdapter(FSubobjectDataAdapterHandle Handle);
+	UE_EXPERIMENTAL(5.7, "Experimental adapter to adjust behavior of SubobjectData")
+	const FSubobjectDataAdapterBase* FindDataAdapter(const FSubobjectData& SubobjectData) const;
+
+	using FSubobjectSubsystemAdapterFilter = TFunction<bool(const TInstancedStruct<FSubobjectDataSubsystemContextDataBase>& ContextData)>;
+	/** */
+	UE_EXPERIMENTAL(5.7, "Experimental adapter to adjust behavior of SubobjectDataSubsystem")
+	UE_API FSubobjectDataSubsystemAdapterHandle RegisterSubsystemAdapter(FSubobjectSubsystemAdapterFilter InFilter, TUniquePtr<FSubobjectDataSubsystemAdapterBase> InAdapter);
+	UE_EXPERIMENTAL(5.7, "Experimental adapter to adjust behavior of SubobjectDataSubsystem")
+	UE_API void UnregisterSubsystemAdapter(FSubobjectDataSubsystemAdapterHandle Handle);
+private:
+	const FSubobjectDataSubsystemAdapterBase* FindSubsystemAdapter(const TInstancedStruct<FSubobjectDataSubsystemContextDataBase>& ContextObject) const;
+public:
+
 	////////////////////////////////////////////
 
 	/**
@@ -409,6 +429,11 @@ private:
 	 */
 	UE_API void FindAllSubobjectData(FSubobjectData* RootContext, TSet<FSubobjectData*>& OutVisited) const;
 	
+	/**
+	 * Creates an invalid SubobjectData and associates a handle
+	 * The SubobjectData can be initialized using the Initialize() function
+	 */
+	UE_API FSubobjectDataHandle CreateSubobjectData();
 	/**
 	 * Gets a ref to subobject data that has been initialized with the given context.
 	 * It will have a unique ID, but may be reused from the FreeData array. 
@@ -432,6 +457,24 @@ private:
 	UE_API FSubobjectDataHandle FindParentForNewSubobject(const UObject* NewSubobject, const FSubobjectDataHandle& SelectedParent);
 
 	FOnNewSubobjectAdded OnNewSubobjectAdded_Delegate;
+
+	struct FSubobjectDataAdapterEntry
+	{
+		TUniquePtr<FSubobjectDataAdapterBase> Adapter;
+		FSubobjectDataAdapterHandle Handle;
+		FSubobjectDataAdapterFilter Filter;
+	};
+	uint64 DataAdapterHandleIndex = 0;
+	TArray<FSubobjectDataAdapterEntry> DataAdapters;
+
+	struct FSubsystemAdapterEntry
+	{
+		TUniquePtr<FSubobjectDataSubsystemAdapterBase> Adapter;
+		FSubobjectDataSubsystemAdapterHandle Handle;
+		FSubobjectSubsystemAdapterFilter Filter;
+	};
+	uint64 SubsystemAdapterHandleIndex = 0;
+	TArray<FSubsystemAdapterEntry> SubsystemAdapters;
 	
 };
 
diff --git a/Engine/Source/Editor/SubobjectDataInterface/SubobjectDataInterface.Build.cs b/Engine/Source/Editor/SubobjectDataInterface/SubobjectDataInterface.Build.cs
--- a/Engine/Source/Editor/SubobjectDataInterface/SubobjectDataInterface.Build.cs
+++ b/Engine/Source/Editor/SubobjectDataInterface/SubobjectDataInterface.Build.cs
@@ -30,6 +30,7 @@ public class SubobjectDataInterface : ModuleRules
 				{
 					"UnrealEd",	// ComponentEditorUtils, AssetBrokerage
 					"GameProjectGeneration",
+					"TypedElementFramework"
 				}
 			);
 		}
diff --git a/Engine/Source/Editor/SubobjectEditor/Private/SSubobjectEditor.cpp b/Engine/Source/Editor/SubobjectEditor/Private/SSubobjectEditor.cpp
--- a/Engine/Source/Editor/SubobjectEditor/Private/SSubobjectEditor.cpp
+++ b/Engine/Source/Editor/SubobjectEditor/Private/SSubobjectEditor.cpp
@@ -183,7 +183,7 @@ FString FSubobjectEditorTreeNode::GetDisplayString() const
 bool FSubobjectEditorTreeNode::IsComponentNode() const
 {
 	FSubobjectData* Data = GetDataSource();
-	return Data ? Data->GetComponentTemplate() != nullptr : false;
+	return Data ? Data->IsComponent() : false;
 }
 
 bool FSubobjectEditorTreeNode::IsNativeComponent() const
@@ -909,6 +909,17 @@ const FSlateBrush* SSubobject_RowWidget::GetIconBrush() const
 	{
 		if (const FSubobjectData* Data = NodePtr->GetDataSource())
 		{
+			TSharedPtr<SSubobjectEditor> SubobjectEditorPinned = SubobjectEditor.Pin();
+			const ISCSEditorUICustomization* UICustomization = SubobjectEditorPinned ? SubobjectEditorPinned->GetUICustomization() : nullptr;
+			if (UICustomization)
+			{
+				const FSlateBrush* SubobjectIcon = UICustomization->GetIconBrush(*Data);
+				if (SubobjectIcon)
+				{
+					return SubobjectIcon;
+				}
+			}
+			
 			if (const AActor* Actor = Data->GetObject<AActor>())
 			{
 				return FClassIconFinder::FindIconForActor(Actor);
@@ -1740,6 +1751,17 @@ TSharedRef<SWidget> SSubobject_RowWidget::GetInheritedLinkWidget()
 	{
 		return SNullWidget::NullWidget;
 	}
+
+	if (TSharedPtr<SSubobjectEditor> PinnedSubobjectEditor = SubobjectEditor.Pin())
+	{
+		if (const ISCSEditorUICustomization* UICustomization = PinnedSubobjectEditor->GetUICustomization())
+		{
+			if (TSharedPtr<SWidget> ControlsWidget = UICustomization->GetControlsWidget(*Data))
+			{
+				return ControlsWidget.ToSharedRef();
+			}
+		}
+	}
 	
 	// Native components are inherited and have a gray hyperlink to their C++ class
 	if(Data->IsNativeComponent())
@@ -2798,6 +2820,11 @@ void SSubobjectEditor::SetUICustomization(TSharedPtr<ISCSEditorUICustomization>
 	UpdateTree(true /*bRegenerateTreeNodes*/);
 }
 
+const ISCSEditorUICustomization* SSubobjectEditor::GetUICustomization() const
+{
+	return UICustomization.Get();
+}
+
 bool SSubobjectEditor::CanCutNodes() const
 {
 	return CanCopyNodes() && CanDeleteNodes();
@@ -2971,13 +2998,14 @@ FSubobjectDataHandle SSubobjectEditor::GetObjectContextHandle() const
 {
 	if (RootNodes.Num() > 0)
 	{
-		ensure(RootNodes[0]->GetObject() == GetObjectContext());
-		return RootNodes[0]->GetDataHandle();
-	}
-	else
-	{
-		return FSubobjectDataHandle::InvalidHandle;
+		if (RootNodes[0]->GetObject() == GetObjectContext())
+		{
+			return RootNodes[0]->GetDataHandle();
+		}
 	}
+
+	return FSubobjectDataHandle::InvalidHandle;
+	
 }
 
 UBlueprint* SSubobjectEditor::GetBlueprint() const
diff --git a/Engine/Source/Editor/SubobjectEditor/Public/SSubobjectEditor.h b/Engine/Source/Editor/SubobjectEditor/Public/SSubobjectEditor.h
--- a/Engine/Source/Editor/SubobjectEditor/Public/SSubobjectEditor.h
+++ b/Engine/Source/Editor/SubobjectEditor/Public/SSubobjectEditor.h
@@ -504,6 +504,9 @@ public:
 	
 	/** Sets UI customizations of this SCSEditor. */
 	UE_API void SetUICustomization(TSharedPtr<ISCSEditorUICustomization> InUICustomization);
+	
+	/** Internally linked function to allow the tree widgets to access customization functions */
+	const ISCSEditorUICustomization* GetUICustomization() const;
 
 	UE_API virtual bool IsEditingAllowed() const;
 	
