Add DeleteSubobjects extension to SubobjectDataSubsystem

diff --git a/Engine/Source/Editor/SubobjectDataInterface/Private/SubobjectDataAdapterBase.cpp b/Engine/Source/Editor/SubobjectDataInterface/Private/SubobjectDataAdapterBase.cpp
--- a/Engine/Source/Editor/SubobjectDataInterface/Private/SubobjectDataAdapterBase.cpp
+++ b/Engine/Source/Editor/SubobjectDataInterface/Private/SubobjectDataAdapterBase.cpp
@@ -35,6 +35,13 @@ void FSubobjectDataSubsystemAdapterBase::GatherSubobjectData(const TInstancedStr
 {
 }
 
+int32 FSubobjectDataSubsystemAdapterBase::DeleteSubobjects(
+	const FSubobjectDataHandle& ContextHandle,
+	const TArray<FSubobjectDataHandle>& SubobjectsToDelete) const
+{
+	return 0;
+}
+
 FSubobjectDataAdapterBase::~FSubobjectDataAdapterBase() = default;
 
 bool FSubobjectDataAdapterBase::CanEdit(const FSubobjectData& Data) const
diff --git a/Engine/Source/Editor/SubobjectDataInterface/Private/SubobjectDataSubsystem.cpp b/Engine/Source/Editor/SubobjectDataInterface/Private/SubobjectDataSubsystem.cpp
--- a/Engine/Source/Editor/SubobjectDataInterface/Private/SubobjectDataSubsystem.cpp
+++ b/Engine/Source/Editor/SubobjectDataInterface/Private/SubobjectDataSubsystem.cpp
@@ -1224,8 +1224,23 @@ int32 USubobjectDataSubsystem::DeleteSubobjects(const FSubobjectDataHandle& Cont
 	{
 		return NumDeletedSubobjects;
 	}
-	const FSubobjectData* ContextData = ContextHandle.GetData();
-	UObject* ContextObj = ContextData->GetMutableObject();
+	const FSubobjectData* SubobjectData = ContextHandle.GetData();
+	if (!SubobjectData)
+	{
+		return NumDeletedSubobjects;
+	}
+	
+	if (const FSubobjectDataSubsystemAdapterBase* Adapter = FindSubsystemAdapter(SubobjectData->GetContextData()))
+	{
+		NumDeletedSubobjects = Adapter->DeleteSubobjects(
+			ContextHandle,
+			SubobjectsToDelete);
+		
+		return NumDeletedSubobjects;
+	}
+
+	
+	UObject* ContextObj = SubobjectData->GetMutableObject();
 	check(ContextObj);
 
 	ContextObj->Modify();
diff --git a/Engine/Source/Editor/SubobjectDataInterface/Public/SubobjectDataAdapterBase.h b/Engine/Source/Editor/SubobjectDataInterface/Public/SubobjectDataAdapterBase.h
--- a/Engine/Source/Editor/SubobjectDataInterface/Public/SubobjectDataAdapterBase.h
+++ b/Engine/Source/Editor/SubobjectDataInterface/Public/SubobjectDataAdapterBase.h
@@ -77,6 +77,7 @@ class FSubobjectDataSubsystemAdapterBase
 public:
 	UE_API virtual ~FSubobjectDataSubsystemAdapterBase();
 	UE_API virtual void GatherSubobjectData(const TInstancedStruct<FSubobjectDataSubsystemContextDataBase>& ContextData, TArray<FSubobjectDataHandle>& OutArray, TFunctionRef<FSubobjectDataHandle()> CreateSubobjectData) const;
+	UE_API virtual int32 DeleteSubobjects(const FSubobjectDataHandle& ContextHandle, const TArray<FSubobjectDataHandle>& SubobjectsToDelete) const;
 };
 
 
