Fix redefined UE_API

diff --git a/Engine/Source/Editor/SubobjectDataInterface/Public/SubobjectDataAdapterBase.h b/Engine/Source/Editor/SubobjectDataInterface/Public/SubobjectDataAdapterBase.h
--- a/Engine/Source/Editor/SubobjectDataInterface/Public/SubobjectDataAdapterBase.h
+++ b/Engine/Source/Editor/SubobjectDataInterface/Public/SubobjectDataAdapterBase.h
@@ -143,4 +143,6 @@ public:
 	UE_API virtual FText GetIntroducedInToolTipText(const FSubobjectData& Data) const;
 
 	UE_API virtual FText GetActorDisplayText(const FSubobjectData& Data) const;
-};
\ No newline at end of file
+};
+
+#undef UE_API
\ No newline at end of file
