[UE] Fix failure to copy instanced BP SCS node component value overrides on copy/paste workflows in the instanced BP Actor details components tree.

Notes:
- This brings the copy path to parity w/ the duplication path.

diff --git a/Engine/Source/Editor/SubobjectDataInterface/Private/SubobjectDataSubsystem.cpp b/Engine/Source/Editor/SubobjectDataInterface/Private/SubobjectDataSubsystem.cpp
--- a/Engine/Source/Editor/SubobjectDataInterface/Private/SubobjectDataSubsystem.cpp
+++ b/Engine/Source/Editor/SubobjectDataInterface/Private/SubobjectDataSubsystem.cpp
@@ -340,7 +340,7 @@ void USubobjectDataSubsystem::GatherSubobjectData(UObject* Context, TArray<FSubo
 		}
 
 		// The components array will be populated with any natively added components
-		// from the constructor/ObjectInitalizer
+		// from the constructor/ObjectInitializer
 		// These components will all be attached to the root component if it exists
 		for (UActorComponent* Component : Components)
 		{
@@ -2343,9 +2343,20 @@ bool USubobjectDataSubsystem::CanCopySubobjects(const TArray<FSubobjectDataHandl
     		{
     			return false;
     		}
+
+			// Get the component template associated with the selected node
+			UActorComponent* ComponentTemplate = nullptr;
+			const FSubobjectDataHandle& ActorRootHandle = GetActorRootHandle(Handle);
+			if (const FSubobjectData* ActorData = ActorRootHandle.GetData())
+			{
+				ComponentTemplate = Data->FindMutableComponentInstanceInActor(ActorData->GetObject<AActor>());
+			}
+			else
+			{
+				ComponentTemplate = Data->GetMutableComponentTemplate();
+			}
     		
-    		// Get the component template associated with the selected node
-    		if (UActorComponent* ComponentTemplate = Data->GetMutableComponentTemplate())
+    		if (ComponentTemplate)
     		{
     			ComponentsToCopy.Add(ComponentTemplate);
     		}
@@ -2370,9 +2381,21 @@ void USubobjectDataSubsystem::CopySubobjects(const TArray<FSubobjectDataHandle>&
 		if(FSubobjectData* Data = Handle.GetData())
 		{
 			ensureMsgf(Data->CanCopy(), TEXT("A non-copiable subobject has been allowed to copy!"));
+
+			// Get the component template associated with the selected node
+			UActorComponent* ComponentTemplate = nullptr;
+			const FSubobjectDataHandle& ActorRootHandle = GetActorRootHandle(Handle);
+			if (const FSubobjectData* ActorData = ActorRootHandle.GetData())
+			{
+				ComponentTemplate = Data->FindMutableComponentInstanceInActor(ActorData->GetObject<AActor>());
+			}
+			else
+			{
+				ComponentTemplate = Data->GetMutableComponentTemplate();
+			}
     		
 			// Get the component template associated with the selected node
-			if (UActorComponent* ComponentTemplate = Data->GetMutableComponentTemplate())
+			if (ComponentTemplate)
 			{
 				ComponentsToCopy.Add(ComponentTemplate);
 				if(BpContext && ComponentTemplate->CreationMethod != EComponentCreationMethod::UserConstructionScript)
