[UE] Fix a regression that disabled contextual copy operations in the Blueprint editor components tree view.

#rnx

diff --git a/Engine/Source/Editor/SubobjectDataInterface/Private/SubobjectDataSubsystem.cpp b/Engine/Source/Editor/SubobjectDataInterface/Private/SubobjectDataSubsystem.cpp
--- a/Engine/Source/Editor/SubobjectDataInterface/Private/SubobjectDataSubsystem.cpp
+++ b/Engine/Source/Editor/SubobjectDataInterface/Private/SubobjectDataSubsystem.cpp
@@ -2351,7 +2351,8 @@ bool USubobjectDataSubsystem::CanCopySubobjects(const TArray<FSubobjectDataHandl
 			{
 				ComponentTemplate = Data->FindMutableComponentInstanceInActor(ActorData->GetObject<AActor>());
 			}
-			else
+			
+			if (!ComponentTemplate)
 			{
 				ComponentTemplate = Data->GetMutableComponentTemplate();
 			}
@@ -2389,7 +2390,8 @@ void USubobjectDataSubsystem::CopySubobjects(const TArray<FSubobjectDataHandle>&
 			{
 				ComponentTemplate = Data->FindMutableComponentInstanceInActor(ActorData->GetObject<AActor>());
 			}
-			else
+			
+			if (!ComponentTemplate)
 			{
 				ComponentTemplate = Data->GetMutableComponentTemplate();
 			}
