Cleanup experimental work on Subobject system to keep API in experimental
* Added a second Initialize function which does not rely on the FSubobjectDataSubsystemContextDataBase struct
* Put the FSubobjectDataSubsystemContextDataBase initialize function in experimental

diff --git a/Engine/Source/Editor/SubobjectDataInterface/Private/SubobjectData.cpp b/Engine/Source/Editor/SubobjectDataInterface/Private/SubobjectData.cpp
--- a/Engine/Source/Editor/SubobjectDataInterface/Private/SubobjectData.cpp
+++ b/Engine/Source/Editor/SubobjectDataInterface/Private/SubobjectData.cpp
@@ -39,10 +39,16 @@ FSubobjectData::FSubobjectData()
 }
 
 FSubobjectData::FSubobjectData(UObject* ContextObject, const FSubobjectDataHandle& InParentHandle, const bool bInIsInheritedSCS)
-	: WeakObjectPtr(ContextObject)
-	, ParentObjectHandle(InParentHandle)
-	, SCSNodePtr(nullptr)
 {
+	Initialize(ContextObject, InParentHandle, bInIsInheritedSCS);
+}
+
+void FSubobjectData::Initialize(UObject* ContextObject, const FSubobjectDataHandle& InParentHandle, const bool bInIsInheritedSCS)
+{
+	WeakObjectPtr = ContextObject;
+	ParentObjectHandle = InParentHandle;
+	SCSNodePtr = nullptr;
+	
 	FSubobjectDataSubsystemContextData_SingleUObjectContextObject SingleUObject(ContextObject);
 	ContextData.InitializeAsScriptStruct(FSubobjectDataSubsystemContextData_SingleUObjectContextObject::StaticStruct(), reinterpret_cast<const uint8*>(&SingleUObject));
 	
diff --git a/Engine/Source/Editor/SubobjectDataInterface/Private/SubobjectDataSubsystem.cpp b/Engine/Source/Editor/SubobjectDataInterface/Private/SubobjectDataSubsystem.cpp
--- a/Engine/Source/Editor/SubobjectDataInterface/Private/SubobjectDataSubsystem.cpp
+++ b/Engine/Source/Editor/SubobjectDataInterface/Private/SubobjectDataSubsystem.cpp
@@ -2757,18 +2757,9 @@ FSubobjectDataHandle USubobjectDataSubsystem::CreateSubobjectData()
 
 FSubobjectDataHandle USubobjectDataSubsystem::CreateSubobjectData(UObject* Context, const FSubobjectDataHandle& ParentHandle /* = FSubobjectDataHandle::InvalidHandle */, bool bIsInheritedSCS/* = false */)
 {
-	PRAGMA_DISABLE_DEPRECATION_WARNINGS
-	TSharedPtr<FSubobjectData> SharedPtr = TSharedPtr<FSubobjectData>(new FSubobjectData(Context, ParentHandle, bIsInheritedSCS));
-	PRAGMA_ENABLE_DEPRECATION_WARNINGS
+	TSharedPtr<FSubobjectData> SharedPtr = TSharedPtr<FSubobjectData>(new FSubobjectData());
+	SharedPtr->Initialize(Context, ParentHandle, bIsInheritedSCS);
 	
-	if(!SharedPtr.IsValid())
-	{
-		ensureMsgf(false, TEXT("The subobject data factories failed to create subobject data!"));
-		PRAGMA_DISABLE_DEPRECATION_WARNINGS
-		SharedPtr = TSharedPtr<FSubobjectData>(new FSubobjectData(Context, ParentHandle, bIsInheritedSCS));
-		PRAGMA_ENABLE_DEPRECATION_WARNINGS
-	}
-
 	check(SharedPtr.IsValid());
 	
 	SharedPtr->Handle.DataPtr = SharedPtr;
diff --git a/Engine/Source/Editor/SubobjectDataInterface/Public/SubobjectData.h b/Engine/Source/Editor/SubobjectDataInterface/Public/SubobjectData.h
--- a/Engine/Source/Editor/SubobjectDataInterface/Public/SubobjectData.h
+++ b/Engine/Source/Editor/SubobjectDataInterface/Public/SubobjectData.h
@@ -33,9 +33,12 @@ public:
 
 	UE_API FSubobjectData();
 
-	UE_DEPRECATED(5.7, "")
+	UE_DEPRECATED(5.7, "Use FSubobjectData default constructor + Initialize function")
 	UE_API FSubobjectData(UObject* ContextObject, const FSubobjectDataHandle& ParentHandle, const bool bIsInheritedSCS);
 
+	UE_API void Initialize(UObject* ContextObject, const FSubobjectDataHandle& ParentHandle, const bool bIsInheritedSCS);
+
+	UE_EXPERIMENTAL(5.7, "")
 	UE_API void Initialize(TInstancedStruct<FSubobjectDataSubsystemContextDataBase> ContextData, const FSubobjectDataHandle& ParentHandle);
 	
 	/** True if this subobject's handle is valid */
@@ -52,6 +55,7 @@ public:
 	*           - FSubobjectDataSubsystemContextData_SingleUObjectContextObject
 	*           - FSubobjectDataSubsystemContextData_TedsRow
 	*/
+	UE_EXPERIMENTAL(5.7, "")
 	UE_API const TInstancedStruct<FSubobjectDataSubsystemContextDataBase>& GetContextData() const;
 	
 	/**
