Add a context menu option to entity components to directly open the file associated with a Verse class in VS Code

Note that this is not setup to work with the card layout

diff --git a/Engine/Source/Editor/SubobjectDataInterface/Private/SubobjectDataAdapterBase.cpp b/Engine/Source/Editor/SubobjectDataInterface/Private/SubobjectDataAdapterBase.cpp
--- a/Engine/Source/Editor/SubobjectDataInterface/Private/SubobjectDataAdapterBase.cpp
+++ b/Engine/Source/Editor/SubobjectDataInterface/Private/SubobjectDataAdapterBase.cpp
@@ -2,6 +2,9 @@
 
 #include "SubobjectDataAdapterBase.h"
 
+#include "DataStorage/Features.h"
+#include "Elements/Columns/TypedElementCompatibilityColumns.h"
+#include "Elements/Interfaces/TypedElementDataStorageInterface.h"
 #include "SubobjectDataHandle.h"
 
 #define LOCTEXT_NAMESPACE "FSubobjectDataBehaviorBase"
@@ -15,6 +18,32 @@ FSubobjectDataSubsystemContextData_SingleUObjectContextObject::FSubobjectDataSub
 {
 }
 
+void UE::SubobjectData::QueryContextDataForObjects(const TInstancedStruct<FSubobjectDataSubsystemContextDataBase>& ContextData, TArray<const UObject*>& OutObjects, const bool bEvenIfPendingKill)
+{
+	if (const FSubobjectDataSubsystemContextData_SingleUObjectContextObject* SingleObjectData = ContextData.GetPtr<FSubobjectDataSubsystemContextData_SingleUObjectContextObject>())
+	{
+		if (const UObject* PinnedObject = SingleObjectData->Object.Get(bEvenIfPendingKill))
+		{
+			OutObjects.Add(PinnedObject);
+		}
+	}
+	else if (const FSubobjectDataSubsystemContextData_TedsRow* TedsRowData = ContextData.GetPtr<FSubobjectDataSubsystemContextData_TedsRow>())
+	{
+		using namespace UE::Editor::DataStorage;
+
+		if (ICoreProvider* EditorDataStorage = GetMutableDataStorageFeature<ICoreProvider>(StorageFeatureName))
+		{
+			if (FTypedElementUObjectColumn* ObjectColumn = EditorDataStorage->GetColumn<FTypedElementUObjectColumn>(TedsRowData->RowHandle))
+			{
+				if (const UObject* ObjectValue = ObjectColumn->Object.Get(bEvenIfPendingKill))
+				{
+					OutObjects.Add(ObjectValue);
+				}
+			}
+		}
+	}
+}
+
 FSubobjectDataSubsystemContextData_TedsRow::FSubobjectDataSubsystemContextData_TedsRow()
 {
 }
diff --git a/Engine/Source/Editor/SubobjectDataInterface/Public/SubobjectDataAdapterBase.h b/Engine/Source/Editor/SubobjectDataInterface/Public/SubobjectDataAdapterBase.h
--- a/Engine/Source/Editor/SubobjectDataInterface/Public/SubobjectDataAdapterBase.h
+++ b/Engine/Source/Editor/SubobjectDataInterface/Public/SubobjectDataAdapterBase.h
@@ -72,6 +72,11 @@ FSubobjectDataSubsystemContextData_TedsRow : public FSubobjectDataSubsystemConte
 	FTedsRowHandle RowHandle;
 };
 
+namespace UE::SubobjectData
+{
+	UE_API void QueryContextDataForObjects(const TInstancedStruct<FSubobjectDataSubsystemContextDataBase>& ContextData, TArray<const UObject*>& OutObjects, bool bEvenIfPendingKill = false);
+} // namespace UE::SubobjectData
+
 class FSubobjectDataSubsystemAdapterBase
 {
 public:
diff --git a/Engine/Source/Editor/SubobjectEditor/Private/SSubobjectEditor.cpp b/Engine/Source/Editor/SubobjectEditor/Private/SSubobjectEditor.cpp
--- a/Engine/Source/Editor/SubobjectEditor/Private/SSubobjectEditor.cpp
+++ b/Engine/Source/Editor/SubobjectEditor/Private/SSubobjectEditor.cpp
@@ -3336,9 +3336,23 @@ void SSubobjectEditor::GetSelectedItemsForContextMenu(TArray<FComponentEventCons
 		{
 			NewItem.Component = const_cast<UActorComponent*>(Data->FindComponentInstanceInActor(ContextAsActor));
 		}
-		else
+		else if (const UObject* BlueprintObject = Data->GetObjectForBlueprint(GetBlueprint()))
+		{
+			NewItem.Component = const_cast<UObject*>(BlueprintObject);
+		}
+		else if (Data->GetContextData().IsValid())
 		{
-			NewItem.Component = const_cast<UObject*>(Data->GetObjectForBlueprint(GetBlueprint()));
+			TArray<const UObject*> ContextDataObjects;
+			UE::SubobjectData::QueryContextDataForObjects(Data->GetContextData(), ContextDataObjects);
+
+			for (const UObject* ContextDataObject : ContextDataObjects)
+			{
+				FComponentEventConstructionData& ContextItem = OutSelectedItems.AddDefaulted_GetRef();
+				ContextItem.Component = const_cast<UObject*>(ContextDataObject);
+			}
+
+			// Skip adding the partially constructed NewItem since we handled that in our loop above
+			continue;
 		}
 		OutSelectedItems.Add(NewItem);
 	}
