Extension points for Copy/Cut/Paste/Duplicate in SubobjectDataSubsystem
* API changed for the SubobjectDataSubsystemAdapter - each operation is interceptable instead of all operations.  Unlike the SubobjectData, the existing API does not provide consistent input data that can be used to determine whether the adapter should intercept.  This forces each function to provide its own interception implementation.
* Added interception hooks for the Can{Copy,Paste,Duplicate} and respective implementations

diff --git a/Engine/Source/Editor/SubobjectDataInterface/Private/SubobjectDataAdapterBase.cpp b/Engine/Source/Editor/SubobjectDataInterface/Private/SubobjectDataAdapterBase.cpp
--- a/Engine/Source/Editor/SubobjectDataInterface/Private/SubobjectDataAdapterBase.cpp
+++ b/Engine/Source/Editor/SubobjectDataInterface/Private/SubobjectDataAdapterBase.cpp
@@ -60,10 +60,31 @@ FSubobjectDataSubsystemContextData_TedsRow::FSubobjectDataSubsystemContextData_T
 
 FSubobjectDataSubsystemAdapterBase::~FSubobjectDataSubsystemAdapterBase() = default;
 
+bool FSubobjectDataSubsystemAdapterBase::CanIntercept_GatherSubobjectData(const TInstancedStruct<FSubobjectDataSubsystemContextDataBase>& ContextData) const
+{
+	return false;
+}
+
 void FSubobjectDataSubsystemAdapterBase::GatherSubobjectData(const TInstancedStruct<FSubobjectDataSubsystemContextDataBase>& ContextData, TArray<FSubobjectDataHandle>& OutArray, TFunctionRef<FSubobjectDataHandle()> CreateSubobjectData) const
 {
 }
 
+bool FSubobjectDataSubsystemAdapterBase::CanIntercept_FindSceneRootForSubobject(const FSubobjectDataHandle& InHandle) const
+{
+	return false;
+}
+
+FSubobjectDataHandle FSubobjectDataSubsystemAdapterBase::FindSceneRootForSubobject(const FSubobjectDataHandle& InHandle) const
+{
+	return FSubobjectDataHandle();
+}
+
+bool FSubobjectDataSubsystemAdapterBase::CanIntercept_DeleteSubobjects(const FSubobjectDataHandle& ContextHandle,
+                                                                       const TArray<FSubobjectDataHandle>& SubobjectsToDelete, UBlueprint* BPContext, bool bForce) const
+{
+	return false;
+}
+
 int32 FSubobjectDataSubsystemAdapterBase::DeleteSubobjects(
 	const FSubobjectDataHandle& ContextHandle,
 	const TArray<FSubobjectDataHandle>& SubobjectsToDelete) const
@@ -71,6 +92,80 @@ int32 FSubobjectDataSubsystemAdapterBase::DeleteSubobjects(
 	return 0;
 }
 
+bool FSubobjectDataSubsystemAdapterBase::CanIntercept_CanCopySubobjects(const TArray<FSubobjectDataHandle>& Handles) const
+{
+	return false;
+}
+
+bool FSubobjectDataSubsystemAdapterBase::CanCopySubobjects(const TArray<FSubobjectDataHandle>& Handles) const
+{
+	return false;
+}
+
+bool FSubobjectDataSubsystemAdapterBase::CanIntercept_CopySubobjects(
+	const TArray<FSubobjectDataHandle>& Handles,
+	UBlueprint* BpContext) const
+{
+	return false;
+}
+
+void FSubobjectDataSubsystemAdapterBase::CopySubobjects(
+	const TArray<FSubobjectDataHandle>& Handles,
+	UBlueprint* BpContext) const
+{
+}
+
+bool FSubobjectDataSubsystemAdapterBase::CanIntercept_CanPasteSubobjects(const FSubobjectDataHandle& RootHandle, UBlueprint* BPContext) const
+{
+	return false;
+}
+
+bool FSubobjectDataSubsystemAdapterBase::CanPasteSubobjects(const FSubobjectDataHandle& RootHandle, UBlueprint* Blueprint) const
+{
+	return false;
+}
+
+bool FSubobjectDataSubsystemAdapterBase::CanIntercept_PasteSubobjects(const FSubobjectDataHandle& RootHandle,
+	const TArray<FSubobjectDataHandle>& NewParentHandles, UBlueprint* BPContext) const
+{
+	return false;
+}
+
+void FSubobjectDataSubsystemAdapterBase::PasteSubobjects(const FSubobjectDataHandle& RootHandle, const TArray<FSubobjectDataHandle>& NewParentHandles,
+	UBlueprint* Blueprint, TArray<FSubobjectDataHandle>& OutPastedHandles) const
+{
+}
+
+bool FSubobjectDataSubsystemAdapterBase::CanIntercept_CanDuplicateSubobjects(
+	const FSubobjectDataHandle& Context,
+	const TArray<FSubobjectDataHandle>& SubobjectsToDup) const
+{
+	return false;
+}
+
+bool FSubobjectDataSubsystemAdapterBase::CanDuplicateSubobjects(
+	const FSubobjectDataHandle& Context,
+	const TArray<FSubobjectDataHandle>& SubobjectsToDup) const
+{
+	return false;
+}
+
+bool FSubobjectDataSubsystemAdapterBase::CanIntercept_DuplicateSubobjects(
+	const FSubobjectDataHandle& Context,
+	const TArray<FSubobjectDataHandle>& SubobjectsToDup,
+	UBlueprint* BpContext) const
+{
+	return false;
+}
+
+void FSubobjectDataSubsystemAdapterBase::DuplicateSubobjects(
+	const FSubobjectDataHandle& Context,
+	const TArray<FSubobjectDataHandle>& SubobjectsToDup,
+	UBlueprint* BpContext,
+	TArray<FSubobjectDataHandle>& OutNewSubobjects) const
+{
+}
+
 FSubobjectDataAdapterBase::~FSubobjectDataAdapterBase() = default;
 
 bool FSubobjectDataAdapterBase::CanEdit(const FSubobjectData& Data) const
diff --git a/Engine/Source/Editor/SubobjectDataInterface/Private/SubobjectDataSubsystem.cpp b/Engine/Source/Editor/SubobjectDataInterface/Private/SubobjectDataSubsystem.cpp
--- a/Engine/Source/Editor/SubobjectDataInterface/Private/SubobjectDataSubsystem.cpp
+++ b/Engine/Source/Editor/SubobjectDataInterface/Private/SubobjectDataSubsystem.cpp
@@ -122,7 +122,6 @@ const FSubobjectDataAdapterBase* USubobjectDataSubsystem::FindDataAdapter(const
 
 
 FSubobjectDataSubsystemAdapterHandle USubobjectDataSubsystem::RegisterSubsystemAdapter(
-	FSubobjectSubsystemAdapterFilter InFilter,
 	TUniquePtr<FSubobjectDataSubsystemAdapterBase> InAdapter)
 {
 	FSubobjectDataSubsystemAdapterHandle Handle;
@@ -132,8 +131,7 @@ FSubobjectDataSubsystemAdapterHandle USubobjectDataSubsystem::RegisterSubsystemA
 	SubsystemAdapters.Emplace(FSubsystemAdapterEntry
 		{
 			.Adapter = MoveTemp(InAdapter),
-			.Handle = Handle,
-			.Filter = MoveTemp(InFilter)
+			.Handle = Handle
 		});
 	return Handle;
 }
@@ -146,19 +144,6 @@ void USubobjectDataSubsystem::UnregisterSubsystemAdapter(FSubobjectDataSubsystem
 	});
 }
 
-const FSubobjectDataSubsystemAdapterBase* USubobjectDataSubsystem::FindSubsystemAdapter(
-	const TInstancedStruct<FSubobjectDataSubsystemContextDataBase>& Struct) const
-{
-	for (const FSubsystemAdapterEntry& Entry : SubsystemAdapters)
-	{
-		if (Entry.Filter(Struct))
-		{
-			return Entry.Adapter.Get();
-		}
-	}
-	return nullptr;
-}
-
 void USubobjectDataSubsystem::K2_GatherSubobjectDataForBlueprint(UBlueprint* Context, TArray<FSubobjectDataHandle>& OutArray)
 {
 	// Return the current CDO that was last generated for the class
@@ -293,13 +278,16 @@ void USubobjectDataSubsystem::GatherSubobjectData(UObject* Context, TArray<FSubo
 	FSubobjectDataSubsystemContextData_SingleUObjectContextObject ContextDataAdapter_OnStack(Context);
 	TConstStructView<FSubobjectDataSubsystemContextDataBase> StructView(ContextDataAdapter_OnStack);
 	TInstancedStruct<FSubobjectDataSubsystemContextDataBase> ContextDataAdapter(StructView);
-	
-	if (const FSubobjectDataSubsystemAdapterBase* Adapter = FindSubsystemAdapter(ContextDataAdapter))
+
+	for (const auto& Entry : SubsystemAdapters)
 	{
-		return Adapter->GatherSubobjectData(
-			ContextDataAdapter,
-			OutArray,
-			[this]() { return this->CreateSubobjectData(); });
+		if (Entry.Adapter->CanIntercept_GatherSubobjectData(ContextDataAdapter))
+		{
+			return Entry.Adapter->GatherSubobjectData(
+				ContextDataAdapter,
+				OutArray,
+				[this]() { return this->CreateSubobjectData(); });
+		}
 	}
 	
 	if (!Context)
@@ -809,6 +797,14 @@ FSubobjectDataHandle USubobjectDataSubsystem::FindSceneRootForSubobject(const FS
 	{
 		return FSubobjectDataHandle::InvalidHandle;
 	}
+
+	for (const auto& Entry : SubsystemAdapters)
+	{
+		if (Entry.Adapter->CanIntercept_FindSceneRootForSubobject(InHandle))
+		{
+			return Entry.Adapter->FindSceneRootForSubobject(InHandle);
+		}
+	}
 	
 	FSubobjectDataHandle ActorHandle = InHandle;
 	FSubobjectData* ActorData = ActorHandle.GetData();
@@ -1229,16 +1225,17 @@ int32 USubobjectDataSubsystem::DeleteSubobjects(const FSubobjectDataHandle& Cont
 	{
 		return NumDeletedSubobjects;
 	}
-	
-	if (const FSubobjectDataSubsystemAdapterBase* Adapter = FindSubsystemAdapter(SubobjectData->GetContextData()))
+
+	for (const auto& Entry : SubsystemAdapters)
 	{
-		NumDeletedSubobjects = Adapter->DeleteSubobjects(
-			ContextHandle,
-			SubobjectsToDelete);
-		
-		return NumDeletedSubobjects;
+		if (Entry.Adapter->CanIntercept_DeleteSubobjects(ContextHandle, SubobjectsToDelete, BPContext, bForce))
+		{
+			NumDeletedSubobjects = Entry.Adapter->DeleteSubobjects(
+				ContextHandle,
+				SubobjectsToDelete);
+			return NumDeletedSubobjects;
+		}
 	}
-
 	
 	UObject* ContextObj = SubobjectData->GetMutableObject();
 	check(ContextObj);
@@ -2333,6 +2330,14 @@ bool USubobjectDataSubsystem::IsValidRename(const FSubobjectDataHandle& Handle,
 
 bool USubobjectDataSubsystem::CanCopySubobjects(const TArray<FSubobjectDataHandle>& Handles) const
 {
+	for (const auto& Entry : SubsystemAdapters)
+	{
+		if (Entry.Adapter->CanIntercept_CanCopySubobjects(Handles))
+		{
+			return Entry.Adapter->CanCopySubobjects(Handles);
+		}
+	}
+		
 	TArray<UActorComponent*> ComponentsToCopy;
 	
 	for (const FSubobjectDataHandle& Handle : Handles)
@@ -2356,20 +2361,29 @@ bool USubobjectDataSubsystem::CanCopySubobjects(const TArray<FSubobjectDataHandl
 			{
 				ComponentTemplate = Data->GetMutableComponentTemplate();
 			}
-    		
-    		if (ComponentTemplate)
-    		{
-    			ComponentsToCopy.Add(ComponentTemplate);
-    		}
-    	}
-    }
+			
+			if (ComponentTemplate)
+			{
+				ComponentsToCopy.Add(ComponentTemplate);
+			}
+		}
+	}
 
-    // Verify that the components can be copied
-    return FComponentEditorUtils::CanCopyComponents(ComponentsToCopy);
+	// Verify that the components can be copied
+	return FComponentEditorUtils::CanCopyComponents(ComponentsToCopy);
 }
 
 void USubobjectDataSubsystem::CopySubobjects(const TArray<FSubobjectDataHandle>& Handles, UBlueprint* BpContext)
 {
+	for (const auto& Entry : SubsystemAdapters)
+	{
+		if (Entry.Adapter->CanIntercept_CopySubobjects(Handles, BpContext))
+		{
+			Entry.Adapter->CopySubobjects(Handles, BpContext);
+			return;
+		}
+	}
+	
 	if(!CanCopySubobjects(Handles))
 	{
 		return;
@@ -2403,18 +2417,18 @@ void USubobjectDataSubsystem::CopySubobjects(const TArray<FSubobjectDataHandle>&
 				if(BpContext && ComponentTemplate->CreationMethod != EComponentCreationMethod::UserConstructionScript)
 				{
 					// CopyComponents uses component attachment to maintain hierarchy, but the SCS templates are not
-                    // setup with a relationship to each other. Briefly setup the attachment between the templates being
-                    // copied so that the hierarchy is retained upon pasting
-                    if (USceneComponent* SceneTemplate = Cast<USceneComponent>(ComponentTemplate))
-                    {
-                    	if (FSubobjectData* SelectedParentNodePtr = Data->GetParentHandle().GetData())
-                    	{
-                    		if (USceneComponent* ParentSceneTemplate = Cast<USceneComponent>(SelectedParentNodePtr->GetMutableComponentTemplate()))
-                    		{
-                    			SceneTemplate->SetupAttachment(ParentSceneTemplate);
-                    		}
-                    	}
-                    }
+					// setup with a relationship to each other. Briefly setup the attachment between the templates being
+					// copied so that the hierarchy is retained upon pasting
+					if (USceneComponent* SceneTemplate = Cast<USceneComponent>(ComponentTemplate))
+					{
+						if (FSubobjectData* SelectedParentNodePtr = Data->GetParentHandle().GetData())
+						{
+							if (USceneComponent* ParentSceneTemplate = Cast<USceneComponent>(SelectedParentNodePtr->GetMutableComponentTemplate()))
+							{
+								SceneTemplate->SetupAttachment(ParentSceneTemplate);
+							}
+						}
+					}
 				}
 			}
 		}
@@ -2441,6 +2455,14 @@ void USubobjectDataSubsystem::CopySubobjects(const TArray<FSubobjectDataHandle>&
 
 bool USubobjectDataSubsystem::CanPasteSubobjects(const FSubobjectDataHandle& RootHandle, UBlueprint* BPContext) const
 {
+	for (const auto& Entry : SubsystemAdapters)
+	{
+		if (Entry.Adapter->CanIntercept_CanPasteSubobjects(RootHandle, BPContext))
+		{
+			return Entry.Adapter->CanPasteSubobjects(RootHandle, BPContext);
+		}
+	}
+	
 	const FSubobjectDataHandle& SceneRootHandle = FindSceneRootForSubobject(RootHandle);
 	const FSubobjectData* RootData = SceneRootHandle.GetData();
 	const USceneComponent* SceneComponent = Cast<USceneComponent>(RootData->GetComponentTemplate());
@@ -2460,6 +2482,15 @@ void USubobjectDataSubsystem::PasteSubobjects(const FSubobjectDataHandle& PasteT
 	{
 		return;
 	}
+
+	for (const auto& Entry : SubsystemAdapters)
+	{
+		if (Entry.Adapter->CanIntercept_PasteSubobjects(PasteToContext, NewParentHandles, Blueprint))
+		{
+			Entry.Adapter->PasteSubobjects(PasteToContext, NewParentHandles, Blueprint, OutPastedHandles);
+			return;
+		}
+	}
 	
 	FSubobjectData* PasteToContextData = PasteToContext.GetSharedDataPtr().Get();	
 	
@@ -2606,6 +2637,20 @@ void USubobjectDataSubsystem::PasteSubobjects(const FSubobjectDataHandle& PasteT
 	}
 }
 
+bool USubobjectDataSubsystem::CanDuplicateSubobjects(const FSubobjectDataHandle& Context, const TArray<FSubobjectDataHandle>& SubobjectsToDup)
+{
+	for (const auto& Entry : SubsystemAdapters)
+	{
+		if (Entry.Adapter->CanIntercept_CanDuplicateSubobjects(Context, SubobjectsToDup))
+		{
+			return Entry.Adapter->CanDuplicateSubobjects(Context, SubobjectsToDup);
+		}
+	}
+
+	// Default behaviour inherited from SSubobjectEditor
+	return CanCopySubobjects(SubobjectsToDup);
+}
+
 void USubobjectDataSubsystem::DuplicateSubobjects(const FSubobjectDataHandle& Context, const TArray<FSubobjectDataHandle>& SubobjectsToDup, UBlueprint* BpContext, TArray<FSubobjectDataHandle>& OutNewSubobjects)
 {
 	if(!Context.IsValid() || SubobjectsToDup.IsEmpty())
@@ -2613,6 +2658,14 @@ void USubobjectDataSubsystem::DuplicateSubobjects(const FSubobjectDataHandle& Co
 		return;
 	}
 
+	for (const auto& Entry : SubsystemAdapters)
+	{
+		if (Entry.Adapter->CanIntercept_DuplicateSubobjects(Context, SubobjectsToDup, BpContext))
+		{
+			return Entry.Adapter->DuplicateSubobjects(Context, SubobjectsToDup, BpContext, OutNewSubobjects);
+		}
+	}
+
 	FAddNewSubobjectParams NewSubobjectParams;
 	NewSubobjectParams.BlueprintContext = BpContext;
 	NewSubobjectParams.ParentHandle = Context;
diff --git a/Engine/Source/Editor/SubobjectDataInterface/Public/SubobjectDataAdapterBase.h b/Engine/Source/Editor/SubobjectDataInterface/Public/SubobjectDataAdapterBase.h
--- a/Engine/Source/Editor/SubobjectDataInterface/Public/SubobjectDataAdapterBase.h
+++ b/Engine/Source/Editor/SubobjectDataInterface/Public/SubobjectDataAdapterBase.h
@@ -77,12 +77,48 @@ namespace UE::SubobjectData
 	UE_API void QueryContextDataForObjects(const TInstancedStruct<FSubobjectDataSubsystemContextDataBase>& ContextData, TArray<const UObject*>& OutObjects, bool bEvenIfPendingKill = false);
 } // namespace UE::SubobjectData
 
+/**
+ * The base of an adapter/plugin that extends the SubobjectDataSubsystem by intercepting
+ * calls to the public API.
+ * To handle other types of subobjects (even non-UObject based), implement this base class and register it with
+ * the SubobjectDataSubsystem instance.
+ * Interception is expected to inspect the input data passed to the CanIntercept_* functions and, if they return true, calls the
+ * appropriate override.  Multiple adapters follow FIFO priority - first registered is first given a chance to intercept.
+ */
 class FSubobjectDataSubsystemAdapterBase
 {
 public:
 	UE_API virtual ~FSubobjectDataSubsystemAdapterBase();
+	UE_API virtual bool CanIntercept_GatherSubobjectData(const TInstancedStruct<FSubobjectDataSubsystemContextDataBase>& ContextData) const;
 	UE_API virtual void GatherSubobjectData(const TInstancedStruct<FSubobjectDataSubsystemContextDataBase>& ContextData, TArray<FSubobjectDataHandle>& OutArray, TFunctionRef<FSubobjectDataHandle()> CreateSubobjectData) const;
+
+	// Call this implementation of FindSceneRootForSubobject if returning true
+	UE_API virtual bool CanIntercept_FindSceneRootForSubobject(const FSubobjectDataHandle& InHandle) const;
+	
+	// Returns the subobject handle of the "scene root" object - the definition dependent on the implementation.
+	// In legacy implementation - returns the actor that owns the provided component ~or~ returns the root SceneComponent of the provided Actor
+	UE_API virtual FSubobjectDataHandle FindSceneRootForSubobject(const FSubobjectDataHandle& InHandle) const;
+	
+	UE_API virtual bool CanIntercept_DeleteSubobjects(const FSubobjectDataHandle& ContextHandle, const TArray<FSubobjectDataHandle>& SubobjectsToDelete, UBlueprint* BPContext = nullptr, bool bForce = false) const;
 	UE_API virtual int32 DeleteSubobjects(const FSubobjectDataHandle& ContextHandle, const TArray<FSubobjectDataHandle>& SubobjectsToDelete) const;
+
+	UE_API virtual bool CanIntercept_CanCopySubobjects(const TArray<FSubobjectDataHandle>& Handles) const;
+	UE_API virtual bool CanCopySubobjects(const TArray<FSubobjectDataHandle>& Handles) const;
+
+	UE_API virtual bool CanIntercept_CopySubobjects(const TArray<FSubobjectDataHandle>& Handles, UBlueprint* BpContext) const;
+	UE_API virtual void CopySubobjects(const TArray<FSubobjectDataHandle>& Handles, UBlueprint* BpContext) const;
+
+	UE_API virtual bool CanIntercept_CanPasteSubobjects(const FSubobjectDataHandle& RootHandle, UBlueprint* BPContext) const;
+	UE_API virtual bool CanPasteSubobjects(const FSubobjectDataHandle& RootHandle,UBlueprint* Blueprint) const;
+
+	UE_API virtual bool CanIntercept_PasteSubobjects(const FSubobjectDataHandle& RootHandle, const TArray<FSubobjectDataHandle>& NewParentHandles, UBlueprint* BPContext) const;
+	UE_API virtual void PasteSubobjects(const FSubobjectDataHandle& RootHandle, const TArray<FSubobjectDataHandle>& NewParentHandles, UBlueprint* Blueprint, TArray<FSubobjectDataHandle>& OutPastedHandles) const;
+	
+	UE_API virtual bool CanIntercept_CanDuplicateSubobjects(const FSubobjectDataHandle& Context, const TArray<FSubobjectDataHandle>& SubobjectsToDup) const;
+	UE_API virtual bool CanDuplicateSubobjects(const FSubobjectDataHandle& Context, const TArray<FSubobjectDataHandle>& SubobjectsToDup) const;
+
+	UE_API virtual bool CanIntercept_DuplicateSubobjects(const FSubobjectDataHandle& Context, const TArray<FSubobjectDataHandle>& SubobjectsToDup, UBlueprint* BpContext) const;
+	UE_API virtual void DuplicateSubobjects(const FSubobjectDataHandle& Context, const TArray<FSubobjectDataHandle>& SubobjectsToDup, UBlueprint* BpContext, TArray<FSubobjectDataHandle>& OutNewSubobjects) const;
 };
 
 
diff --git a/Engine/Source/Editor/SubobjectDataInterface/Public/SubobjectDataSubsystem.h b/Engine/Source/Editor/SubobjectDataInterface/Public/SubobjectDataSubsystem.h
--- a/Engine/Source/Editor/SubobjectDataInterface/Public/SubobjectDataSubsystem.h
+++ b/Engine/Source/Editor/SubobjectDataInterface/Public/SubobjectDataSubsystem.h
@@ -137,12 +137,9 @@ public:
 	using FSubobjectSubsystemAdapterFilter = TFunction<bool(const TInstancedStruct<FSubobjectDataSubsystemContextDataBase>& ContextData)>;
 	/** */
 	UE_EXPERIMENTAL(5.7, "Experimental adapter to adjust behavior of SubobjectDataSubsystem")
-	UE_API FSubobjectDataSubsystemAdapterHandle RegisterSubsystemAdapter(FSubobjectSubsystemAdapterFilter InFilter, TUniquePtr<FSubobjectDataSubsystemAdapterBase> InAdapter);
+	UE_API FSubobjectDataSubsystemAdapterHandle RegisterSubsystemAdapter(TUniquePtr<FSubobjectDataSubsystemAdapterBase> InAdapter);
 	UE_EXPERIMENTAL(5.7, "Experimental adapter to adjust behavior of SubobjectDataSubsystem")
 	UE_API void UnregisterSubsystemAdapter(FSubobjectDataSubsystemAdapterHandle Handle);
-private:
-	const FSubobjectDataSubsystemAdapterBase* FindSubsystemAdapter(const TInstancedStruct<FSubobjectDataSubsystemContextDataBase>& ContextObject) const;
-public:
 
 	////////////////////////////////////////////
 
@@ -378,6 +375,8 @@ public:
 	*/
 	UFUNCTION(BlueprintCallable, Category = "SubobjectDataSubsystem")
 	UE_API void PasteSubobjects(const FSubobjectDataHandle& PasteToContext, const TArray<FSubobjectDataHandle>& NewParentHandles, UBlueprint* BpContext, TArray<FSubobjectDataHandle>& OutPastedHandles);
+
+	UE_API bool CanDuplicateSubobjects(const FSubobjectDataHandle& Context, const TArray<FSubobjectDataHandle>& SubobjectsToDup);
 	
 	/**
 	 * Duplicate the given array of subobjects on the context. 
@@ -471,7 +470,6 @@ private:
 	{
 		TUniquePtr<FSubobjectDataSubsystemAdapterBase> Adapter;
 		FSubobjectDataSubsystemAdapterHandle Handle;
-		FSubobjectSubsystemAdapterFilter Filter;
 	};
 	uint64 SubsystemAdapterHandleIndex = 0;
 	TArray<FSubsystemAdapterEntry> SubsystemAdapters;
diff --git a/Engine/Source/Editor/SubobjectEditor/Private/SSubobjectEditor.cpp b/Engine/Source/Editor/SubobjectEditor/Private/SSubobjectEditor.cpp
--- a/Engine/Source/Editor/SubobjectEditor/Private/SSubobjectEditor.cpp
+++ b/Engine/Source/Editor/SubobjectEditor/Private/SSubobjectEditor.cpp
@@ -2938,6 +2938,15 @@ bool SSubobjectEditor::CanDuplicateComponent() const
 		return false;
 	}
 	
+	TArray<FSubobjectDataHandle> SelectedHandles = GetSelectedHandles();
+	if(USubobjectDataSubsystem* System = USubobjectDataSubsystem::Get())
+	{
+		if (!RootNodes.IsEmpty())
+		{
+			return System->CanDuplicateSubobjects(RootNodes[0]->GetDataHandle(), SelectedHandles);
+		}
+	}
+
 	return CanCopyNodes();
 }
 
