Editor: Fix some child pointer issues when initializing subobject data.

diff --git a/Engine/Source/Editor/SubobjectDataInterface/Private/SubobjectData.cpp b/Engine/Source/Editor/SubobjectDataInterface/Private/SubobjectData.cpp
--- a/Engine/Source/Editor/SubobjectDataInterface/Private/SubobjectData.cpp
+++ b/Engine/Source/Editor/SubobjectDataInterface/Private/SubobjectData.cpp
@@ -98,6 +98,10 @@ void FSubobjectData::Initialize(TInstancedStruct<FSubobjectDataSubsystemContextD
 
 	ContextData = MoveTemp(InContextData);
 	ParentObjectHandle = ParentHandle;
+	if (ParentHandle.IsValid())
+	{
+		ParentHandle.GetData()->AddChildHandleOnly(GetHandle());
+	}
 }
 
 bool FSubobjectData::CanEdit() const
@@ -998,6 +1002,15 @@ void FSubobjectData::DisposeUnsafe()
 	ParentObjectHandle = FSubobjectDataHandle();
 }
 
+void FSubobjectData::DisconnectFromParent()
+{
+	if (ParentObjectHandle.IsValid())
+	{
+		ParentObjectHandle.GetData()->RemoveChildHandleOnly(GetHandle());
+	}
+	ParentObjectHandle = FSubobjectDataHandle::InvalidHandle;
+}
+
 const UObject* FSubobjectData::GetObject(bool bEvenIfPendingKill) const
 {
 	return WeakObjectPtr.Get(bEvenIfPendingKill);
diff --git a/Engine/Source/Editor/SubobjectDataInterface/Public/SubobjectData.h b/Engine/Source/Editor/SubobjectDataInterface/Public/SubobjectData.h
--- a/Engine/Source/Editor/SubobjectDataInterface/Public/SubobjectData.h
+++ b/Engine/Source/Editor/SubobjectDataInterface/Public/SubobjectData.h
@@ -189,7 +189,7 @@ public:
 	bool HasParent() const { return ParentObjectHandle.IsValid(); }
 	const TArray<FSubobjectDataHandle>& GetChildrenHandles() const { return ChildrenHandles; }
 	
-	/** Walks the parent heirarchy and returns the root subobject in this case */
+	/** Walks the parent hierarchy and returns the root subobject in this case */
 	UE_API FSubobjectDataHandle GetRootSubobject() const;
 
 	// Return true if the given handle is a child of this subobject
@@ -236,6 +236,12 @@ public:
 	 *	involved in cycles, but live FSubobjectDatas will start to behave strangely once it is called
 	 */
 	UE_API void DisposeUnsafe();
+
+	/**
+	 * Removes this FSubobjectData from its parent's children, and clears the parent pointer.
+	 */
+	UE_EXPERIMENTAL(5.8, "This function may be changed.")
+	UE_API void DisconnectFromParent();
 private:
 	/**
 	* Add the given subobject handle as a child of this.
