[Bug Fix][Small CR] Blueprint - The X axis in Rotation is set to "-0.0" and fails to be set to "0.0"

[Fix] Identity quaternion can convert to a zero rotator with a -0.0 Roll value. This fix corrects that.

#Synced-CL 49621724

diff --git a/Engine/Source/Editor/SubobjectDataInterface/Private/SubobjectDataSubsystem.cpp b/Engine/Source/Editor/SubobjectDataInterface/Private/SubobjectDataSubsystem.cpp
--- a/Engine/Source/Editor/SubobjectDataInterface/Private/SubobjectDataSubsystem.cpp
+++ b/Engine/Source/Editor/SubobjectDataInterface/Private/SubobjectDataSubsystem.cpp
@@ -784,7 +784,15 @@ static void ConformTransformRelativeToParent(USceneComponent* SceneComponentTemp
 	// Store new relative rotation value (if not set to absolute)
 	if (!SceneComponentTemplate->IsUsingAbsoluteRotation())
 	{
-		SceneComponentTemplate->SetRelativeRotation_Direct(RelativeTM.Rotator());
+		FRotator RelativeRotation = RelativeTM.Rotator();
+
+		// Identity quaternion can convert to a zero rotator with a -0.0 Roll value, correct that
+		if (RelativeRotation == FRotator{ 0.0, 0.0, -0.0 })
+		{
+			RelativeRotation = FRotator{ 0.0 };
+		}
+		
+		SceneComponentTemplate->SetRelativeRotation_Direct(RelativeRotation);
 	}
 
 	// Store new relative scale value (if not set to absolute)
