Fix Subobject editor accidently adding a component using the in world component editing flow via python script

[RN] USubobjectDataSubsystem::AddNewComponent will no longer crash when adding components without an explicit bluepritn context

#Synced-CL 50069009

diff --git a/Engine/Source/Editor/SubobjectDataInterface/Private/SubobjectData.cpp b/Engine/Source/Editor/SubobjectDataInterface/Private/SubobjectData.cpp
--- a/Engine/Source/Editor/SubobjectDataInterface/Private/SubobjectData.cpp
+++ b/Engine/Source/Editor/SubobjectDataInterface/Private/SubobjectData.cpp
@@ -401,9 +401,26 @@ UBlueprint* FSubobjectData::GetBlueprintBeingEdited() const
 	{
 		return BP;
 	}
-	else if (Object->IsTemplate())
+	else
 	{
-		return GetBlueprint();
+		// walk to the root handle and see if it's a CDO or blueprint root object:
+		FSubobjectDataHandle RootHandle = GetRootSubobject();
+		if (const FSubobjectData* RootData = RootHandle.GetData())
+		{
+			UObject* Obj = RootData->WeakObjectPtr.Get();
+			if(Obj && Obj->HasAnyFlags(RF_ClassDefaultObject))
+			{
+				return Cast<UBlueprint>(Obj->GetClass()->ClassGeneratedBy);
+			}
+			else if(const UBlueprintGeneratedClass* BPGC = Cast<UBlueprintGeneratedClass>(Obj))
+			{
+				return Cast<UBlueprint>(BPGC->ClassGeneratedBy);
+			}
+			else
+			{
+				return Cast<UBlueprint>(Obj);
+			}
+		}
 	}
 
 	return nullptr;
diff --git a/Engine/Source/Editor/SubobjectDataInterface/Private/SubobjectDataSubsystem.cpp b/Engine/Source/Editor/SubobjectDataInterface/Private/SubobjectDataSubsystem.cpp
--- a/Engine/Source/Editor/SubobjectDataInterface/Private/SubobjectDataSubsystem.cpp
+++ b/Engine/Source/Editor/SubobjectDataInterface/Private/SubobjectDataSubsystem.cpp
@@ -992,7 +992,8 @@ FSubobjectDataHandle USubobjectDataSubsystem::AddNewSubobject(const FAddNewSubob
 	// an SCS node will use the SCS template data, which means that we'll find a valid Blueprint context in that case. For editor scripting, we want it
 	// to always modify the Blueprint, but for "normal" editor operations (e.g. drag-and-drop), we need to bypass this if the root is an Actor instance.
 	UBlueprint* Blueprint = Params.BlueprintContext;
-	if (!Blueprint && !GetActorRootHandle(ParentObjHandle).IsValid())
+	const FSubobjectDataHandle& ParentActorRootHandle = GetActorRootHandle(ParentObjHandle);
+	if (!Blueprint)
 	{
 		// maybe the parent handle has a bp for context:
 		Blueprint = ParentObjHandle.GetData()->GetBlueprintBeingEdited();
